Copying //UE4/Orion-Staging to //UE4/Main (Source: //Orion/Dev-General @ 3316439)

diff --git a/Engine/Plugins/Runtime/HTNPlanner/Source/HTNTestSuite/Private/HTNTest.cpp b/Engine/Plugins/Runtime/HTNPlanner/Source/HTNTestSuite/Private/HTNTest.cpp
new file mode 100644
--- /dev/null
+++ b/Engine/Plugins/Runtime/HTNPlanner/Source/HTNTestSuite/Private/HTNTest.cpp
@@ -0,0 +1,616 @@
+// Copyright 1998-2017 Epic Games, Inc. All Rights Reserved.
+
+#include "CoreMinimal.h"
+#include "HTNBuilder.h"
+#include "HTNPlanner.h"
+#include "AI/HTNBrainComponent.h"
+#include "MockHTN.h"
+#include "AITestsCommon.h"
+#include "Misc/App.h"
+#include "Debug/HTNDebug.h"
+
+#define LOCTEXT_NAMESPACE "AITestSuite_HTNTest"
+
+PRAGMA_DISABLE_OPTIMIZATION
+
+// missing tests:
+//  testing invalid conditions
+//	testing specific scenarios with known plans 
+//	test step by step generating the same plan as GeneratePlan
+//	using memory counters see if anything's leaking
+
+struct FHTNTestBase : public FAITestBase
+{
+	FHTNBuilder_Domain DomainBuilder;
+	FHTNWorldState WorldState;
+	FHTNPlanner Planner;
+
+	void PopulateWorldState()
+	{
+		for (int32 WSIndex = 0; WSIndex < int32(EMockHTNWorldState::MAX); ++WSIndex)
+		{
+			// setting every key to it's numerical index value
+			WorldState.SetValueUnsafe(WSIndex, WSIndex);
+		}
+	}
+
+	void PopulateDomain(const bool bCompile = true)
+	{
+		DomainBuilder.SetRootName(TEXT("Root"));
+		{
+			FHTNBuilder_CompositeTask& CompositeTaskBuilder = DomainBuilder.AddCompositeTask(TEXT("Root"));
+			{
+				FHTNBuilder_Method& MethodsBuilder = CompositeTaskBuilder.AddMethod(
+					TArray<FHTNCondition>({
+					FHTNCondition(EMockHTNWorldState::EnemyHealth, EHTNWorldStateCheck::Greater).SetRHSAsValue(0)
+					, FHTNCondition(EMockHTNWorldState::EnemyActor, EHTNWorldStateCheck::IsTrue)
+				}));
+				MethodsBuilder.AddTask(TEXT("AttackEnemy"));
+			}
+			{
+				FHTNBuilder_Method& MethodsBuilder = CompositeTaskBuilder.AddMethod();
+				MethodsBuilder.AddTask(TEXT("FindPatrolPoint"));
+				MethodsBuilder.AddTask(TEXT("NavigateToMoveDestination"));
+			}
+		}
+		{
+			FHTNBuilder_CompositeTask& CompositeTaskBuilder = DomainBuilder.AddCompositeTask(TEXT("AttackEnemy"));
+			{
+				FHTNBuilder_Method& MethodsBuilder = CompositeTaskBuilder.AddMethod(FHTNCondition(EMockHTNWorldState::HasWeapon, EHTNWorldStateCheck::IsTrue));
+				MethodsBuilder.AddTask(TEXT("NavigateToEnemy"));
+				MethodsBuilder.AddTask(TEXT("UseWeapon"));
+				MethodsBuilder.AddTask(TEXT("Root"));
+			}
+			{
+				FHTNBuilder_Method& MethodsBuilder = CompositeTaskBuilder.AddMethod();
+				MethodsBuilder.AddTask(TEXT("FindWeapon"));
+				MethodsBuilder.AddTask(TEXT("NavigateToWeapon"));
+				MethodsBuilder.AddTask(TEXT("PickUp"));
+				MethodsBuilder.AddTask(TEXT("AttackEnemy"));
+			}
+		}
+		{
+			FHTNBuilder_PrimitiveTask& PrimitiveTaskBuilder = DomainBuilder.AddPrimitiveTask(TEXT("FindPatrolPoint"));
+			PrimitiveTaskBuilder.SetOperator(EMockHTNTaskOperator::FindPatrolPoint, EMockHTNWorldState::MoveDestination);
+		}
+		{
+			FHTNBuilder_PrimitiveTask& PrimitiveTaskBuilder = DomainBuilder.AddPrimitiveTask(TEXT("FindWeapon"));
+			PrimitiveTaskBuilder.SetOperator(EMockHTNTaskOperator::FindWeapon, EMockHTNWorldState::PickupLocation);
+		}
+		{
+			FHTNBuilder_PrimitiveTask& PrimitiveTaskBuilder = DomainBuilder.AddPrimitiveTask(TEXT("NavigateToMoveDestination"));
+			PrimitiveTaskBuilder.SetOperator(EMockHTNTaskOperator::NavigateTo, EMockHTNWorldState::MoveDestination);	// Local Variables?
+			PrimitiveTaskBuilder.AddEffect(FHTNEffect(EMockHTNWorldState::CurrentLocation, EHTNWorldStateOperation::Set).SetRHSAsWSKey(EMockHTNWorldState::MoveDestination));
+		}
+		{
+			FHTNBuilder_PrimitiveTask& PrimitiveTaskBuilder = DomainBuilder.AddPrimitiveTask(TEXT("NavigateToEnemy"));
+			PrimitiveTaskBuilder.SetOperator(EMockHTNTaskOperator::NavigateTo, EMockHTNWorldState::EnemyActor);
+			PrimitiveTaskBuilder.AddEffect(FHTNEffect(EMockHTNWorldState::CurrentLocation, EHTNWorldStateOperation::Set).SetRHSAsWSKey(EMockHTNWorldState::EnemyActor));
+			PrimitiveTaskBuilder.AddEffect(FHTNEffect(EMockHTNWorldState::CanSeeEnemy, EHTNWorldStateOperation::Set).SetRHSAsValue(1));
+		}
+		{
+			FHTNBuilder_PrimitiveTask& PrimitiveTaskBuilder = DomainBuilder.AddPrimitiveTask(TEXT("NavigateToWeapon"));
+			PrimitiveTaskBuilder.SetOperator(EMockHTNTaskOperator::NavigateTo, EMockHTNWorldState::PickupLocation);
+			PrimitiveTaskBuilder.AddEffect(FHTNEffect(EMockHTNWorldState::CurrentLocation, EHTNWorldStateOperation::Set).SetRHSAsWSKey(EMockHTNWorldState::PickupLocation));
+		}
+		{
+			FHTNBuilder_PrimitiveTask& PrimitiveTaskBuilder = DomainBuilder.AddPrimitiveTask(TEXT("PickUp"));
+			PrimitiveTaskBuilder.SetOperator(EMockHTNTaskOperator::PickUp, EMockHTNWorldState::PickupLocation);
+			PrimitiveTaskBuilder.AddEffect(FHTNEffect(EMockHTNWorldState::HasWeapon, EHTNWorldStateOperation::Set).SetRHSAsValue(1));
+		}
+		{
+			FHTNBuilder_PrimitiveTask& PrimitiveTaskBuilder = DomainBuilder.AddPrimitiveTask(TEXT("UseWeapon"));
+			PrimitiveTaskBuilder.SetOperator(EMockHTNTaskOperator::UseWeapon, EMockHTNWorldState::EnemyActor);
+			PrimitiveTaskBuilder.AddEffect(FHTNEffect(EMockHTNWorldState::Ammo, EHTNWorldStateOperation::Decrease).SetRHSAsValue(1));
+			PrimitiveTaskBuilder.AddEffect(FHTNEffect(EMockHTNWorldState::EnemyHealth, EHTNWorldStateOperation::Decrease).SetRHSAsValue(1));
+		}
+
+		if (bCompile)
+		{
+			DomainBuilder.Compile();
+		}
+	}
+};
+
+//////////////////////////////////////////////////////////////////////////
+
+struct FAITest_HTNDomainBuilderBasics : public FAITestBase
+{
+	FAITest_HTNDomainBuilderBasics()
+	{
+		FHTNBuilder_Domain DomainBuilder;
+
+		Test(TEXT("Initially DomainBuilder instance should be empty"), DomainBuilder.CompositeTasks.Num() == 0 && DomainBuilder.PrimitiveTasks.Num() == 0);
+
+		const uint32 CompointTasksCount = 5;
+		for (int32 CompositeTaskIndex = 0; CompositeTaskIndex < CompointTasksCount; ++CompositeTaskIndex)
+		{
+			FHTNBuilder_CompositeTask& CompositeTaskBuilder = DomainBuilder.AddCompositeTask(*FString::Printf(TEXT("c_%d"), CompositeTaskIndex));
+			for (int32 MethodIndex = 0; MethodIndex < (CompositeTaskIndex % 3); ++MethodIndex)
+			{
+				FHTNBuilder_Method* MethodsBuilder = nullptr;
+
+				if ((CompositeTaskIndex + MethodIndex) != 0)
+				{
+					TArray<FHTNCondition> Conditions;
+					for (int32 ConditionIndex = 0; ConditionIndex < (CompositeTaskIndex % 3); ++ConditionIndex)
+					{
+						Conditions.Add(FHTNCondition(ConditionIndex, EHTNWorldStateCheck(ConditionIndex)));
+					}
+
+					MethodsBuilder = &CompositeTaskBuilder.AddMethod(Conditions);
+				}
+				else
+				{
+					// this is testing the non-array method
+					MethodsBuilder = &CompositeTaskBuilder.AddMethod(FHTNCondition(0, EHTNWorldStateCheck(0)));
+				}
+
+				for (int32 TaskIndex = 0; TaskIndex < (CompositeTaskIndex % 4); ++TaskIndex)
+				{
+					MethodsBuilder->AddTask(*FString::Printf(TEXT("t_%d"), CompositeTaskIndex));
+				}
+			}
+		}
+
+		for (int32 CompositeTaskIndex = 0; CompositeTaskIndex < DomainBuilder.CompositeTasks.Num(); ++CompositeTaskIndex)
+		{
+			FHTNBuilder_CompositeTask* CompositeTaskBuilder = DomainBuilder.CompositeTasks.Find(*FString::Printf(TEXT("c_%d"), CompositeTaskIndex));
+			Test(FString::Printf(TEXT("Failed to find Composite task c_%d that has just been added"), CompositeTaskIndex), CompositeTaskBuilder != nullptr);
+
+			if (CompositeTaskBuilder != nullptr)
+			{
+				Test(FString::Printf(TEXT("Method count mismatch for c_%d"), CompositeTaskIndex), CompositeTaskBuilder->Methods.Num() == (CompositeTaskIndex % 3));
+				for (int32 MethodIndex = 0; MethodIndex < CompositeTaskBuilder->Methods.Num(); ++MethodIndex)
+				{
+					if ((CompositeTaskIndex + MethodIndex) != 0)
+					{
+						Test(FString::Printf(TEXT("Condition count mismatch for c_%d[%d] method (array path)"), CompositeTaskIndex, MethodIndex)
+							, CompositeTaskBuilder->Methods[MethodIndex].Conditions.Num() == (CompositeTaskIndex % 3));
+					}
+					else
+					{
+						Test(FString::Printf(TEXT("Condition count mismatch for c_%d[%d] method (single instance)"), CompositeTaskIndex, MethodIndex)
+							, CompositeTaskBuilder->Methods[MethodIndex].Conditions.Num() == 1);
+					}
+
+					Test(FString::Printf(TEXT("Task count mismatch for c_%d[%d] method"), CompositeTaskIndex, MethodIndex)
+						, CompositeTaskBuilder->Methods[MethodIndex].Tasks.Num() == (CompositeTaskIndex % 4));
+				}
+			}			
+		}
+
+		const uint32 PrimitiveTasksCount = 4;
+		for (int32 PrimitiveTaskIndex = 0; PrimitiveTaskIndex < PrimitiveTasksCount; ++PrimitiveTaskIndex)
+		{
+			FHTNBuilder_PrimitiveTask& PrimitiveTaskBuilder = DomainBuilder.AddPrimitiveTask(*FString::Printf(TEXT("p_%d"), PrimitiveTaskIndex));
+			PrimitiveTaskBuilder.SetOperator(PrimitiveTaskIndex, PrimitiveTaskIndex * 2);
+			for (int32 EffectIndex = 0; EffectIndex < (PrimitiveTaskIndex % 3); ++EffectIndex)
+			{
+				PrimitiveTaskBuilder.AddEffect(FHTNEffect(EffectIndex, EHTNWorldStateOperation::Set));
+			}
+		}
+
+		Test(TEXT("Wrong number of primitive tasks added"), DomainBuilder.PrimitiveTasks.Num() == PrimitiveTasksCount);
+		for (int32 PrimitiveTaskIndex = 0; PrimitiveTaskIndex < PrimitiveTasksCount; ++PrimitiveTaskIndex)
+		{
+			FHTNBuilder_PrimitiveTask* PrimitiveTaskBuilder = DomainBuilder.PrimitiveTasks.Find(*FString::Printf(TEXT("p_%d"), PrimitiveTaskIndex));
+			Test(FString::Printf(TEXT("Failed to find primitive task p_%d that has just been added"), PrimitiveTaskIndex), PrimitiveTaskBuilder != nullptr);
+
+			if (PrimitiveTaskBuilder)
+			{
+				Test(FString::Printf(TEXT("Primitive task p_%d operator is wrong"), PrimitiveTasksCount), PrimitiveTaskBuilder->ActionID == PrimitiveTaskIndex);
+				Test(FString::Printf(TEXT("Primitive task p_%d effects count is wrong"), PrimitiveTasksCount), PrimitiveTaskBuilder->Effects.Num() == (PrimitiveTaskIndex % 3));
+			}
+		}
+	}
+};
+IMPLEMENT_AI_LATENT_TEST(FAITest_HTNDomainBuilderBasics, "System.Engine.AI.HTN.DomainBuilderBasics")
+
+struct FAITest_HTNBuildDomain : public FHTNTestBase
+{
+	FAITest_HTNBuildDomain()
+	{
+		PopulateDomain();
+
+		Test(TEXT("DomainBuilder stores wrong number of primitive tasks"), DomainBuilder.PrimitiveTasks.Num() == 7);
+		Test(TEXT("DomainBuilder stores wrong number of Composite tasks"), DomainBuilder.CompositeTasks.Num() == 2);
+		
+		FHTNBuilder_CompositeTask* RootBuilder = DomainBuilder.GetRootAsCompositeTask();
+		Test(TEXT("Root task should be set"), RootBuilder != nullptr);
+		if (RootBuilder && RootBuilder->Methods.Num() > 0)
+		{
+			Test(TEXT("Root task\'s method [0] should be configured to have two conditions"), RootBuilder->Methods[0].Conditions.Num() == 2);
+		}
+		
+		TArray<FHTNBuilder_CompositeTask> CompositeTasks;
+		DomainBuilder.CompositeTasks.GenerateValueArray(CompositeTasks);
+		Test(TEXT("DomainBuilder stores wrong number of methods for the first Composite task"), CompositeTasks[0].Methods.Num() == 2);
+		Test(TEXT("DomainBuilder stores wrong number of methods for the second Composite task"), CompositeTasks[1].Methods.Num() == 2);
+	}
+};
+IMPLEMENT_AI_LATENT_TEST(FAITest_HTNBuildDomain, "System.Engine.AI.HTN.BuildDomain")
+
+struct FAITest_HTNPlanning : public FHTNTestBase
+{
+	FAITest_HTNPlanning()
+	{
+		FHTNResult Result;
+		
+		Planner.GeneratePlan(*(DomainBuilder.DomainInstance), WorldState, Result);
+		Test(TEXT("Planning with an empty domain should result in an empty plan"), Result.TaskIDs.Num() == 0);
+
+		PopulateDomain();
+
+		Planner.GeneratePlan(*(DomainBuilder.DomainInstance), WorldState, Result);
+
+		TArray<FHTNBuilder_CompositeTask> CompositeTasks;
+		DomainBuilder.CompositeTasks.GenerateValueArray(CompositeTasks);
+		// Patrol plan
+		Test(TEXT("Patrol plan should be same length as the last methods of the root Composite task"), CompositeTasks[0].Methods.Last().Tasks.Num() == Result.TaskIDs.Num());
+		for (int32 TaskIndex = 0; TaskIndex < Result.TaskIDs.Num(); ++TaskIndex)
+		{
+			const FHTNPolicy::FTaskID TaskID = Result.TaskIDs[TaskIndex];
+			Test(TEXT("Patrol plan element mismatch"), TaskID == DomainBuilder.DomainInstance->FindTaskID(CompositeTasks[0].Methods.Last().Tasks[TaskIndex]));
+		}
+
+		const FHTNDomain EmptyDomain;
+		Planner.GeneratePlan(EmptyDomain, WorldState, Result);
+		Test(TEXT("Reusing previous planning result with an empty domain should result in an empty plan"), Result.TaskIDs.Num() == 0);
+
+		WorldState.ApplyEffect(FHTNEffect(EMockHTNWorldState::EnemyHealth, EHTNWorldStateOperation::Set).SetRHSAsValue(1));
+		WorldState.ApplyEffect(FHTNEffect(EMockHTNWorldState::EnemyActor, EHTNWorldStateOperation::Set).SetRHSAsValue(1));
+		Planner.GeneratePlan(*(DomainBuilder.DomainInstance), WorldState, Result);
+		Test(TEXT(""), Result.TaskIDs.Num() > 0);
+	}
+};
+IMPLEMENT_AI_LATENT_TEST(FAITest_HTNPlanning, "System.Engine.AI.HTN.Planning")
+
+struct FAITest_HTNPlanningRollback : public FHTNTestBase
+{
+	FAITest_HTNPlanningRollback()
+	{		
+		FHTNResult Result;
+
+		// build a domain that will force rolling back
+		// first method should get accepted in the first planner step
+		//		then one of the tasks it consists of should be a Composite task, that fails its condition
+		// Note: WorldState is populated with 0 values and has 128 keys (by default)
+		{
+			FHTNBuilder_CompositeTask& CompositeTaskBuilder = DomainBuilder.AddCompositeTask(NAME_None);	// root
+			{
+				FHTNBuilder_Method& MethodsBuilder = CompositeTaskBuilder.AddMethod();
+				MethodsBuilder.AddTask(TEXT("FailedComposite"));
+			}
+			{
+				FHTNBuilder_Method& MethodsBuilder = CompositeTaskBuilder.AddMethod();
+				MethodsBuilder.AddTask(TEXT("SuccessfulComposite"));
+			}
+		}
+		{
+			FHTNBuilder_CompositeTask& CompositeTaskBuilder = DomainBuilder.AddCompositeTask(TEXT("FailedComposite"));
+			{
+				FHTNBuilder_Method& MethodsBuilder = CompositeTaskBuilder.AddMethod(FHTNCondition(0, EHTNWorldStateCheck::Greater).SetRHSAsValue(0));
+				MethodsBuilder.AddTask(TEXT("DummyPrimitive1"));
+			}
+		}
+		{
+			FHTNBuilder_CompositeTask& CompositeTaskBuilder = DomainBuilder.AddCompositeTask(TEXT("SuccessfulComposite"));
+			{
+				FHTNBuilder_Method& MethodsBuilder = CompositeTaskBuilder.AddMethod(FHTNCondition(0, EHTNWorldStateCheck::Equal).SetRHSAsValue(0));
+				MethodsBuilder.AddTask(TEXT("DummyPrimitive2"));
+			}
+		}
+		DomainBuilder.AddPrimitiveTask(TEXT("DummyPrimitive1"));
+		const FName DummyPrimitiveName2 = TEXT("DummyPrimitive2");
+		DomainBuilder.AddPrimitiveTask(DummyPrimitiveName2);
+		
+		DomainBuilder.Compile();
+
+		Planner.GeneratePlan(*(DomainBuilder.DomainInstance), WorldState, Result);
+		Test(TEXT("First Rollback plan should consist of one task, DummyPrimitive2"), Result.TaskIDs.Num() == 1 
+			&& FHTNDebug::GetTaskName(DomainBuilder, Result.TaskIDs[0]) == DummyPrimitiveName2);
+	}
+};
+IMPLEMENT_AI_LATENT_TEST(FAITest_HTNPlanningRollback, "System.Engine.AI.HTN.PlanningRollback")
+
+struct FAITest_HTNDecompileDomain : public FHTNTestBase
+{
+	FAITest_HTNDecompileDomain()
+	{
+		Test(TEXT("Compiling an empty domain is allowed"), DomainBuilder.Compile());
+
+		PopulateDomain();
+
+		FHTNBuilder_Domain DomainBuilder2(DomainBuilder.DomainInstance);
+		DomainBuilder2.Decompile();
+		const FString OriginalDescription = DomainBuilder.GetDebugDescription();
+		const FString DecompiledDescription = DomainBuilder2.GetDebugDescription();
+
+		Test(TEXT("Decompilation should result in identical DomainBuilder"), OriginalDescription.Equals(DecompiledDescription));
+	}
+};
+IMPLEMENT_AI_LATENT_TEST(FAITest_HTNDecompileDomain, "System.Engine.AI.HTN.DomainDecompilation")
+
+struct FAITest_HTNDomainCompilationIssues : public FHTNTestBase
+{
+	FAITest_HTNDomainCompilationIssues()
+	{
+		const FName MissingTaskName = TEXT("MissingTask");
+
+		FHTNBuilder_CompositeTask& CompositeTaskBuilder = DomainBuilder.AddCompositeTask(TEXT("Root"));	// root
+		{
+			FHTNBuilder_Method& MethodsBuilder = CompositeTaskBuilder.AddMethod();
+			MethodsBuilder.AddTask(MissingTaskName);
+		}
+
+		Test(TEXT("Domain with missing tasks should not compile"), DomainBuilder.Compile() == false);
+		Test(TEXT("Domain should be empty after a failed compilation"), DomainBuilder.DomainInstance->IsEmpty());
+
+		DomainBuilder.AddPrimitiveTask(MissingTaskName);
+		Test(TEXT("After adding missing task domain should compile just fine"), DomainBuilder.Compile());
+		Test(TEXT("Domain should not be empty after a successful compilation"), DomainBuilder.DomainInstance->IsEmpty() == false);
+	}
+};
+IMPLEMENT_AI_LATENT_TEST(FAITest_HTNDomainCompilationIssues, "System.Engine.AI.HTN.DomainDecompilationIssues")
+
+struct FAITest_HTNWorldRepresentation : public FHTNTestBase
+{
+	FAITest_HTNWorldRepresentation()
+	{
+		for (int32 WSIndex = 0; WSIndex < int32(EMockHTNWorldState::MAX); ++WSIndex)
+		{
+			FHTNPolicy::FWSValue Value;
+			Test(TEXT("Retrieving known values from the WorldState instance"), WorldState.GetValue(WSIndex, Value) && Value == FHTNPolicy::DefaultValue);
+		}
+
+		PopulateWorldState();
+		
+		const uint32 ReferenceValue = 3;
+		for (int32 WSIndex = 0; WSIndex < int32(EMockHTNWorldState::MAX); ++WSIndex)
+		{
+			FHTNPolicy::FWSValue Value;
+			Test(TEXT("Retrieving known values from the WorldState instance"), WorldState.GetValue(WSIndex, Value) && Value == WSIndex);
+
+			for (int32 OpIndex = 0; OpIndex < int32(EHTNWorldStateCheck::MAX); ++OpIndex)
+			{
+				const EHTNWorldStateCheck OpCode = EHTNWorldStateCheck(OpIndex);
+				bool bExpectedResult = false;
+
+				switch (OpCode)
+				{
+				case EHTNWorldStateCheck::Less:
+					bExpectedResult = Value < ReferenceValue;
+					break;
+				case EHTNWorldStateCheck::LessOrEqual:
+					bExpectedResult = Value <= ReferenceValue;
+					break;
+				case EHTNWorldStateCheck::Equal:
+					bExpectedResult = Value == ReferenceValue;
+					break;
+				case EHTNWorldStateCheck::NotEqual:
+					bExpectedResult = Value != ReferenceValue;
+					break;
+				case EHTNWorldStateCheck::GreaterOrEqual:
+					bExpectedResult = Value >= ReferenceValue;
+					break;
+				case EHTNWorldStateCheck::Greater:
+					bExpectedResult = Value > ReferenceValue;
+					break;
+				case EHTNWorldStateCheck::IsTrue:
+					bExpectedResult = (Value != 0);
+					break;
+				default:
+					Test(TEXT("Unhanled operation ID!"), true);
+				}
+
+				FString Message = FString::Printf(TEXT("Testing %s on $d"), *FHTNDebug::HTNWorldStateCheckToString(OpCode), Value);
+				Test(Message, WorldState.CheckCondition(FHTNCondition(WSIndex, OpCode).SetRHSAsValue(ReferenceValue)) == bExpectedResult);
+			}
+		}
+	}
+};
+IMPLEMENT_AI_LATENT_TEST(FAITest_HTNWorldRepresentation, "System.Engine.AI.HTN.WorldRepresentation")
+
+struct FAITest_HTNCondition : public FHTNTestBase
+{
+	FAITest_HTNCondition()
+	{
+		PopulateWorldState();
+
+		const uint32 ReferenceValue = 3;
+
+		for (int32 WSIndex = 0; WSIndex < int32(EMockHTNWorldState::MAX); ++WSIndex)
+		{
+			for (int32 Value = 0; Value < int32(EMockHTNWorldState::MAX); ++Value)
+			{
+				const FHTNPolicy::FWSValue AsValue = FHTNPolicy::FWSValue(Value);
+				const FHTNPolicy::FWSKey AsKey = FHTNPolicy::FWSKey(Value);
+
+				Test(FString::Printf(TEXT("Condition WS[%d] < %d"), WSIndex, AsValue)
+					, WorldState.CheckCondition(FHTNCondition(WSIndex, EHTNWorldStateCheck::Less).SetRHSAsValue(AsValue)) == (WSIndex < AsValue));
+				Test(FString::Printf(TEXT("Condition WS[%d] <= %d"), WSIndex, AsValue)
+					, WorldState.CheckCondition(FHTNCondition(WSIndex, EHTNWorldStateCheck::LessOrEqual).SetRHSAsValue(AsValue)) == (WSIndex <= AsValue));
+				Test(FString::Printf(TEXT("Condition WS[%d] == %d"), WSIndex, AsValue)
+					, WorldState.CheckCondition(FHTNCondition(WSIndex, EHTNWorldStateCheck::Equal).SetRHSAsValue(AsValue)) == (WSIndex == AsValue));
+				Test(FString::Printf(TEXT("Condition WS[%d] != %d"), WSIndex, AsValue)
+					, WorldState.CheckCondition(FHTNCondition(WSIndex, EHTNWorldStateCheck::NotEqual).SetRHSAsValue(AsValue)) == (WSIndex != AsValue));
+				Test(FString::Printf(TEXT("Condition WS[%d] >= %d"), WSIndex, AsValue)
+					, WorldState.CheckCondition(FHTNCondition(WSIndex, EHTNWorldStateCheck::GreaterOrEqual).SetRHSAsValue(AsValue)) == (WSIndex >= AsValue));
+				Test(FString::Printf(TEXT("Condition WS[%d] > %d"), WSIndex, AsValue)
+					, WorldState.CheckCondition(FHTNCondition(WSIndex, EHTNWorldStateCheck::Greater).SetRHSAsValue(AsValue)) == (WSIndex > AsValue));
+
+				Test(FString::Printf(TEXT("Condition WS[%d] < WS[%d]"), WSIndex, AsKey)
+					, WorldState.CheckCondition(FHTNCondition(WSIndex, EHTNWorldStateCheck::Less).SetRHSAsWSKey(AsKey)) == (WSIndex < AsKey));
+				Test(FString::Printf(TEXT("Condition WS[%d] <= WS[%d]"), WSIndex, AsKey)
+					, WorldState.CheckCondition(FHTNCondition(WSIndex, EHTNWorldStateCheck::LessOrEqual).SetRHSAsWSKey(AsKey)) == (WSIndex <= AsKey));
+				Test(FString::Printf(TEXT("Condition WS[%d] == WS[%d]"), WSIndex, AsKey)
+					, WorldState.CheckCondition(FHTNCondition(WSIndex, EHTNWorldStateCheck::Equal).SetRHSAsWSKey(AsKey)) == (WSIndex == AsKey));
+				Test(FString::Printf(TEXT("Condition WS[%d] != WS[%d]"), WSIndex, AsKey)
+					, WorldState.CheckCondition(FHTNCondition(WSIndex, EHTNWorldStateCheck::NotEqual).SetRHSAsWSKey(AsKey)) == (WSIndex != AsKey));
+				Test(FString::Printf(TEXT("Condition WS[%d] >= WS[%d]"), WSIndex, AsKey)
+					, WorldState.CheckCondition(FHTNCondition(WSIndex, EHTNWorldStateCheck::GreaterOrEqual).SetRHSAsWSKey(AsKey)) == (WSIndex >= AsKey));
+				Test(FString::Printf(TEXT("Condition WS[%d] > WS[%d]"), WSIndex, AsKey)
+					, WorldState.CheckCondition(FHTNCondition(WSIndex, EHTNWorldStateCheck::Greater).SetRHSAsWSKey(AsKey)) == (WSIndex > AsKey));
+			}
+		}
+	}
+};
+IMPLEMENT_AI_LATENT_TEST(FAITest_HTNCondition, "System.Engine.AI.HTN.Condition")
+
+struct FAITest_HTNMethodSelection : public FAITestBase 
+{
+	FAITest_HTNMethodSelection()
+	{
+		FHTNDomain Domain;
+
+		//FHTNCompositeTask& RootCompositeTask = Domain.AddCompositeTask(TEXT("Root"));
+		// do planning and come up with empty plan
+
+		// INJECTION
+	}
+};
+IMPLEMENT_AI_LATENT_TEST(FAITest_HTNMethodSelection, "System.Engine.AI.HTN.MethodSelection")
+
+struct FAITest_HTNTrivialPlanning : public FHTNTestBase 
+{
+	FAITest_HTNTrivialPlanning()
+	{
+		//FHTNDomain Domain;
+
+		TArray<FHTNCondition> Conditions;
+		//FHTNMethodSelection MethodSelection(0, Conditions);
+
+
+		//FHTNCompositeTask& RootCompositeTask = Domain.AddCompositeTask(TEXT("Root"));
+		// do planning and come up with empty plan
+	}
+};
+IMPLEMENT_AI_LATENT_TEST(FAITest_HTNTrivialPlanning, "System.Engine.AI.HTN.TrivialPlanning")
+
+struct FAITest_HTNExtendingDomain : public FAITestBase
+{
+	FAITest_HTNExtendingDomain()
+	{
+		//FHTNDomain Domain;
+
+		//FHTNCompositeTask& RootCompositeTask = Domain.AddCompositeTask(TEXT("Root"));
+		// do planning and come up with empty plan
+
+		// INJECTION
+	}
+};
+IMPLEMENT_AI_LATENT_TEST(FAITest_HTNExtendingDomain, "System.Engine.AI.HTN.ExtendingDomain")
+
+struct FAITest_HTNCustomWSCheck : public FHTNTestBase
+{
+	static bool bCheckFunctionRun;
+	static bool CustomCheck(const FHTNPolicy::FWSValue* Values, const FHTNCondition& Condition)
+	{
+		bCheckFunctionRun = true;
+		return true;
+	}
+
+	FAITest_HTNCustomWSCheck()
+	{
+		const uint32 CustomCheckID = FHTNWorldStateOperations::RegisterCustomCheckType(&FAITest_HTNCustomWSCheck::CustomCheck, TEXT("CustomCheck")); 
+
+		FHTNBuilder_CompositeTask& CompositeTaskBuilder = DomainBuilder.AddCompositeTask(NAME_None);	// root
+		{
+			FHTNBuilder_Method& MethodsBuilder = CompositeTaskBuilder.AddMethod(FHTNCondition(0, CustomCheckID));
+			MethodsBuilder.AddTask(TEXT("Task1"));
+		}
+		FHTNBuilder_PrimitiveTask& PrimitiveTask = DomainBuilder.AddPrimitiveTask(TEXT("Task1"));
+		
+		DomainBuilder.Compile();
+		
+		FHTNResult Result;
+		Planner.GeneratePlan(*DomainBuilder.DomainInstance, WorldState, Result);
+
+		Test(TEXT("Custom check has been exeduted"), bCheckFunctionRun);
+		Test(TEXT("The custom check should allow for construction of the plan"), Result.TaskIDs.Num() == 1);
+	}
+};
+bool FAITest_HTNCustomWSCheck::bCheckFunctionRun = false;
+IMPLEMENT_AI_LATENT_TEST(FAITest_HTNCustomWSCheck, "System.Engine.AI.HTN.CustomWSCheck")
+
+struct FAITest_HTNCustomWSOperation : public FHTNTestBase
+{
+	static uint32 StaticValue;
+	static void CustomOperation(FHTNPolicy::FWSValue* Values, const FHTNEffect& Effect)
+	{
+		++StaticValue;
+		Values[Effect.KeyLeftHand] = StaticValue * 1024;
+	}
+
+	FAITest_HTNCustomWSOperation()
+	{
+		const uint32 CustomOperationID = FHTNWorldStateOperations::RegisterCustomOperationType(&FAITest_HTNCustomWSOperation::CustomOperation, TEXT("CustomOperation"));
+
+		FHTNBuilder_CompositeTask& CompositeTaskBuilder = DomainBuilder.AddCompositeTask(NAME_None);	// root
+		{
+			FHTNBuilder_Method& MethodsBuilder = CompositeTaskBuilder.AddMethod();
+			MethodsBuilder.AddTask(TEXT("Task1"));
+		}
+		FHTNBuilder_PrimitiveTask& PrimitiveTask = DomainBuilder.AddPrimitiveTask(TEXT("Task1"));
+		PrimitiveTask.AddEffect(FHTNEffect(0, CustomOperationID));
+		PrimitiveTask.AddEffect(FHTNEffect(2, CustomOperationID));
+
+		DomainBuilder.Compile();
+
+		FHTNResult Result;
+		Planner.GeneratePlan(*DomainBuilder.DomainInstance, WorldState, Result);
+
+		Test(TEXT("Custom Operation has been exeduted"), StaticValue == 2);
+		Test(TEXT("Checking custom effect on key 0"), Planner.GetWorldState().GetValueUnsafe(0) == 1024);
+		Test(TEXT("Checking custom effect on key 0"), Planner.GetWorldState().GetValueUnsafe(2) == 1024 * 2);
+	}
+};
+uint32 FAITest_HTNCustomWSOperation::StaticValue = 0;
+IMPLEMENT_AI_LATENT_TEST(FAITest_HTNCustomWSOperation, "System.Engine.AI.HTN.CustomWSOperation")
+
+//	compare contents of TaskIDs in HTNResult to the operators it contains
+struct FAITest_HTNOperatorsOfGeneratedPlan : public FHTNTestBase
+{
+	FAITest_HTNOperatorsOfGeneratedPlan()
+	{
+		FHTNBuilder_CompositeTask& CompositeTaskBuilder = DomainBuilder.AddCompositeTask(NAME_None);
+		{
+			FHTNBuilder_Method& MethodsBuilder = CompositeTaskBuilder.AddMethod();
+			MethodsBuilder.AddTask(TEXT("Task2"));
+			MethodsBuilder.AddTask(TEXT("Task1"));
+		}
+		FHTNBuilder_PrimitiveTask& PrimitiveTask1 = DomainBuilder.AddPrimitiveTask(TEXT("Task1"));
+		PrimitiveTask1.SetOperator(1, 2);
+		FHTNBuilder_PrimitiveTask& PrimitiveTask2 = DomainBuilder.AddPrimitiveTask(TEXT("Task2"));
+		PrimitiveTask2.SetOperator(3, 4);
+		
+		DomainBuilder.Compile();
+
+		FHTNResult Result;
+		Planner.GeneratePlan(*DomainBuilder.DomainInstance, WorldState, Result);
+		
+		Test(TEXT("Plan should contain two elements"), Result.TaskIDs.Num() == 2);
+		if (Result.TaskIDs.Num() >= 2)
+		{
+			const FHTNPolicy::FTaskID TaskID1 = DomainBuilder.DomainInstance->FindTaskID(TEXT("Task1"));
+			const FHTNPolicy::FTaskID TaskID2 = DomainBuilder.DomainInstance->FindTaskID(TEXT("Task2"));
+			Test(TEXT("Task2 should be the first one"), Result.TaskIDs[0] == TaskID2);
+			Test(TEXT("Task2 action should be the first one"), Result.ActionsSequence[0].ActionID == 3 && Result.ActionsSequence[0].Parameter == 4);
+			Test(TEXT("Task1 should be the second one"), Result.TaskIDs[1] == TaskID1);
+			Test(TEXT("Task1 action should be the first one"), Result.ActionsSequence[1].ActionID == 1 && Result.ActionsSequence[1].Parameter == 2);
+
+
+		}
+	}
+};
+IMPLEMENT_AI_LATENT_TEST(FAITest_HTNOperatorsOfGeneratedPlan, "System.Engine.AI.HTN.OperatorsOfGeneratedPlan")
+
+//----------------------------------------------------------------------//
+// Component tests
+//----------------------------------------------------------------------//
+typedef FAITest_SimpleComponentBasedTest<UMockHTNComponent> FAITest_HTNComponentTest;
+
+PRAGMA_ENABLE_OPTIMIZATION
+
+#undef LOCTEXT_NAMESPACE
diff --git a/Engine/Plugins/Runtime/HTNPlanner/Source/HTNTestSuite/Private/HTNTestSuiteModule.cpp b/Engine/Plugins/Runtime/HTNPlanner/Source/HTNTestSuite/Private/HTNTestSuiteModule.cpp
new file mode 100644
--- /dev/null
+++ b/Engine/Plugins/Runtime/HTNPlanner/Source/HTNTestSuite/Private/HTNTestSuiteModule.cpp
@@ -0,0 +1,13 @@
+// Copyright 1998-2017 Epic Games, Inc. All Rights Reserved.
+
+#include "HTNTestSuiteModule.h"
+
+#define LOCTEXT_NAMESPACE "HTNTestSuite"
+
+class FHTNTestSuiteModule : public IHTNTestSuiteModule
+{
+};
+
+IMPLEMENT_MODULE(FHTNTestSuiteModule, HTNTestSuiteModule)
+
+#undef LOCTEXT_NAMESPACE
diff --git a/Engine/Source/Developer/AITestSuite/Classes/MockAI.h b/Engine/Source/Developer/AITestSuite/Classes/MockAI.h
--- a/Engine/Source/Developer/AITestSuite/Classes/MockAI.h
+++ b/Engine/Source/Developer/AITestSuite/Classes/MockAI.h
@@ -6,7 +6,7 @@
 #include "Stats/Stats.h"
 #include "UObject/ObjectMacros.h"
 #include "UObject/Object.h"
-#include "Developer/AITestSuite/Private/Tests/AITestsCommon.h"
+#include "AITestsCommon.h"
 #include "Tickable.h"
 
 #include "MockAI.generated.h"
diff --git a/Engine/Source/Developer/AITestSuite/Private/AITestSuite.cpp b/Engine/Source/Developer/AITestSuite/Private/AITestSuite.cpp
--- a/Engine/Source/Developer/AITestSuite/Private/AITestSuite.cpp
+++ b/Engine/Source/Developer/AITestSuite/Private/AITestSuite.cpp
@@ -1,7 +1,7 @@
 // Copyright 1998-2017 Epic Games, Inc. All Rights Reserved.
 
 #include "AITestSuite.h"
-#include "Tests/AITestsCommon.h"
+#include "AITestsCommon.h"
 
 DEFINE_LOG_CATEGORY(LogAITestSuite);
 DEFINE_LOG_CATEGORY(LogBehaviorTreeTest);
diff --git a/Engine/Source/Developer/AITestSuite/Private/BehaviorTree/TestBTDecorator_DelayedAbort.cpp b/Engine/Source/Developer/AITestSuite/Private/BehaviorTree/TestBTDecorator_DelayedAbort.cpp
--- a/Engine/Source/Developer/AITestSuite/Private/BehaviorTree/TestBTDecorator_DelayedAbort.cpp
+++ b/Engine/Source/Developer/AITestSuite/Private/BehaviorTree/TestBTDecorator_DelayedAbort.cpp
@@ -1,7 +1,7 @@
 // Copyright 1998-2017 Epic Games, Inc. All Rights Reserved.
 
 #include "BehaviorTree/TestBTDecorator_DelayedAbort.h"
-#include "Tests/AITestsCommon.h"
+#include "AITestsCommon.h"
 
 UTestBTDecorator_DelayedAbort::UTestBTDecorator_DelayedAbort(const FObjectInitializer& ObjectInitializer) : Super(ObjectInitializer)
 {
diff --git a/Engine/Source/Developer/AITestSuite/Public/AITestsCommon.h b/Engine/Source/Developer/AITestSuite/Public/AITestsCommon.h
new file mode 100644
--- /dev/null
+++ b/Engine/Source/Developer/AITestSuite/Public/AITestsCommon.h
@@ -0,0 +1,178 @@
+// Copyright 1998-2017 Epic Games, Inc. All Rights Reserved.
+
+#pragma once
+
+#include "CoreMinimal.h"
+#include "Misc/AutomationTest.h"
+#include "TestLogger.h"
+#include "Actions/PawnActionsComponent.h"
+
+#define ENSURE_FAILED_TESTS 1
+
+class UBehaviorTree;
+class UMockAI_BT;
+
+DECLARE_LOG_CATEGORY_EXTERN(LogAITestSuite, Log, All);
+DECLARE_LOG_CATEGORY_EXTERN(LogBehaviorTreeTest, Log, All);
+
+DEFINE_LATENT_AUTOMATION_COMMAND_ONE_PARAMETER(FAITestCommand_WaitSeconds, float, Duration);
+
+class FAITestCommand_WaitOneTick : public IAutomationLatentCommand
+{
+public: 
+	FAITestCommand_WaitOneTick()
+		: bAlreadyRun(false)
+	{} 
+	virtual bool Update() override;
+private: 
+	bool bAlreadyRun;
+};
+
+
+namespace FAITestHelpers
+{
+	UWorld* GetWorld();
+	static const float TickInterval = 1.f / 30;
+
+	void UpdateFrameCounter();
+	uint64 FramesCounter();
+}
+
+struct AITESTSUITE_API FAITestBase
+{
+private:
+	// internals
+	TArray<UObject*> SpawnedObjects;	
+	uint32 bTearedDown : 1;
+protected:
+	FAutomationTestBase* TestRunner;
+
+	FAITestBase() : bTearedDown(false), TestRunner(nullptr)
+	{}
+
+	template<typename ClassToSpawn>
+	ClassToSpawn* NewAutoDestroyObject()
+	{
+		ClassToSpawn* ObjectInstance = NewObject<ClassToSpawn>();
+		ObjectInstance->AddToRoot();
+		SpawnedObjects.Add(ObjectInstance);
+		return ObjectInstance;
+	}
+
+	UWorld& GetWorld() const
+	{
+		UWorld* World = FAITestHelpers::GetWorld();
+		check(World);
+		return *World;
+	}
+
+	// loggin helper
+	void Test(const FString& Description, bool bValue);
+
+public:
+
+	virtual void SetTestInstance(FAutomationTestBase& AutomationTestInstance) { TestRunner = &AutomationTestInstance; }
+
+	// interface
+	virtual ~FAITestBase();
+	virtual void SetUp() {}
+	/** @return true to indicate that the test is done. */
+	virtual bool Update() { return true; }
+	virtual void InstantTest() {}
+	// must be called!
+	virtual void TearDown();
+};
+
+DEFINE_EXPORTED_LATENT_AUTOMATION_COMMAND_ONE_PARAMETER(AITESTSUITE_API, FAITestCommand_SetUpTest, FAITestBase*, AITest);
+DEFINE_EXPORTED_LATENT_AUTOMATION_COMMAND_ONE_PARAMETER(AITESTSUITE_API, FAITestCommand_PerformTest, FAITestBase*, AITest);
+DEFINE_EXPORTED_LATENT_AUTOMATION_COMMAND_ONE_PARAMETER(AITESTSUITE_API, FAITestCommand_TearDownTest, FAITestBase*, AITest);
+
+// @note that TestClass needs to derive from FAITestBase
+#define IMPLEMENT_AI_LATENT_TEST(TestClass, PrettyName) \
+	IMPLEMENT_SIMPLE_AUTOMATION_TEST(TestClass##_Runner, PrettyName, (EAutomationTestFlags::ClientContext | EAutomationTestFlags::EditorContext | EAutomationTestFlags::EngineFilter)) \
+	bool TestClass##_Runner::RunTest(const FString& Parameters) \
+	{ \
+		/* spawn test instance. Setup should be done in test's constructor */ \
+		TestClass* TestInstance = new TestClass(); \
+		TestInstance->SetTestInstance(*this); \
+		/* set up */ \
+		ADD_LATENT_AUTOMATION_COMMAND(FAITestCommand_SetUpTest(TestInstance)); \
+		/* run latent command to update */ \
+		ADD_LATENT_AUTOMATION_COMMAND(FAITestCommand_PerformTest(TestInstance)); \
+		/* run latent command to tear down */ \
+		ADD_LATENT_AUTOMATION_COMMAND(FAITestCommand_TearDownTest(TestInstance)); \
+		return true; \
+	} 
+
+#define IMPLEMENT_AI_INSTANT_TEST(TestClass, PrettyName) \
+	IMPLEMENT_SIMPLE_AUTOMATION_TEST(TestClass##Runner, PrettyName, (EAutomationTestFlags::ClientContext | EAutomationTestFlags::EditorContext | EAutomationTestFlags::EngineFilter)) \
+	bool TestClass##Runner::RunTest(const FString& Parameters) \
+	{ \
+		/* spawn test instance. Setup should be done in test's constructor */ \
+		TestClass* TestInstance = new TestClass(); \
+		TestInstance->SetTestInstance(*this); \
+		/* set up */ \
+		TestInstance->SetUp(); \
+		/* call the instant-test code */ \
+		TestInstance->InstantTest(); \
+		/* tear down */ \
+		TestInstance->TearDown(); \
+		return true; \
+	} 
+
+//----------------------------------------------------------------------//
+// Specific test types
+//----------------------------------------------------------------------//
+
+struct FAITest_SimpleBT : public FAITestBase
+{
+	TArray<int32> ExpectedResult;
+	UBehaviorTree* BTAsset;
+	UMockAI_BT* AIBTUser;
+	bool bUseSystemTicking;
+
+	FAITest_SimpleBT();	
+	
+	virtual void SetUp() override;
+	virtual bool Update() override;
+	
+	virtual void VerifyResults();
+};
+
+template<class TComponent>
+struct FAITest_SimpleComponentBasedTest : public FAITestBase
+{
+	FTestLogger<int32> Logger;
+	TComponent* Component;
+
+	FAITest_SimpleComponentBasedTest()
+	{
+		Component = NewAutoDestroyObject<TComponent>();
+	}
+
+	virtual void SetTestInstance(FAutomationTestBase& AutomationTestInstance) override
+	{ 
+		FAITestBase::SetTestInstance(AutomationTestInstance);
+		Logger.TestRunner = TestRunner;
+	}
+
+	virtual ~FAITest_SimpleComponentBasedTest()
+	{
+		Test(TEXT("Not all expected values has been logged"), Logger.ExpectedValues.Num() == 0 || Logger.ExpectedValues.Num() == Logger.LoggedValues.Num());
+	}
+
+	virtual void SetUp() override
+	{
+		UWorld* World = FAITestHelpers::GetWorld();
+		Component->RegisterComponentWithWorld(World);
+	}
+
+	void TickComponent()
+	{
+		Component->TickComponent(FAITestHelpers::TickInterval, ELevelTick::LEVELTICK_All, nullptr);
+	}
+
+	//virtual bool Update() override;
+};
+
+typedef FAITest_SimpleComponentBasedTest<UPawnActionsComponent> FAITest_SimpleActionsTest;
diff --git a/Engine/Source/Developer/AutomationController/Private/AutomationCommandline.cpp b/Engine/Source/Developer/AutomationController/Private/AutomationCommandline.cpp
--- a/Engine/Source/Developer/AutomationController/Private/AutomationCommandline.cpp
+++ b/Engine/Source/Developer/AutomationController/Private/AutomationCommandline.cpp
@@ -10,6 +10,9 @@
 #include "Modules/ModuleManager.h"
 #include "Misc/FilterCollection.h"
 #include "Interfaces/IAutomationControllerModule.h"
+#include "FileManager.h"
+#include "Paths.h"
+#include "FileHelper.h"
 
 
 /** States for running the automation process */
@@ -32,6 +35,7 @@ namespace EAutomationCommand
 		ListAllTests,				//List all tests for the session
 		//RunSingleTest,			//Run one test specified by the commandline
 		RunCommandLineTests,		//Run only tests that are listed on the commandline
+		RunCheckpointTests,			// Run only tests listed on the commandline with checkpoints in case of a crash.
 		RunAll,						//Run all the tests that are supported
 		RunFilter,                  //
 		Quit						//quit the app when tests are done
@@ -118,6 +122,7 @@ public:
 		}
 		return false;
 	}
+
 	
 	void GenerateTestNamesFromCommandLine(const TArray<FString>& AllTestNames, TArray<FString>& OutTestNames)
 	{
@@ -149,6 +154,24 @@ public:
 				}
 			}
 		}
+		// If we have the TestsRun array set up and are using the same command as before, clear out already run tests. 
+		if (TestsRun.Num() > 0)
+		{
+			if (TestsRun[0] == StringCommand)
+			{
+				for (int i = 1; i < TestsRun.Num(); i++)	
+				{
+					if (OutTestNames.Remove(TestsRun[i]))
+					{
+						OutputDevice->Logf(TEXT("Skipping %s due to Checkpoint."), *TestsRun[i]);
+					}
+				}
+			}
+			else
+			{
+				AutomationController->CleanUpCheckpointFile();
+			}
+		}
 	}
 
 	void FindWorkers(float DeltaTime)
@@ -203,6 +226,29 @@ public:
 				AutomationTestState = EAutomationTestState::Complete;
 			}
 		}
+		else if (AutomationCommand == EAutomationCommand::RunCheckpointTests)
+		{
+			TArray <FString> FilteredTestNames;
+			GenerateTestNamesFromCommandLine(AllTestNames, FilteredTestNames);
+			if (FilteredTestNames.Num())
+			{
+				AutomationController->StopTests();
+				AutomationController->SetEnabledTests(FilteredTestNames);
+				if (TestsRun.Num())
+				{
+					AutomationController->WriteLoadedCheckpointDataToFile();
+				}
+				else
+				{
+					AutomationController->WriteLineToCheckpointFile(StringCommand);
+				}
+				bRunTests = true;
+			}
+			else
+			{
+				AutomationTestState = EAutomationTestState::Complete;
+			}
+		}
 		else if (AutomationCommand == EAutomationCommand::RunFilter)
 		{
 			if (FilterMaps.Contains(StringCommand))
@@ -356,6 +402,73 @@ public:
 					Ar.Logf(TEXT("Running all tests matching substring: %s"), *StringCommand);
 					AutomationCommandQueue.Add(EAutomationCommand::RunCommandLineTests);
 				}
+				else if (FParse::Command(&TempCmd, TEXT("RunCheckpointedTests")))
+				{
+					StringCommand = TempCmd;
+					Ar.Logf(TEXT("Running all tests with checkpoints matching substring: %s"), *StringCommand);
+					AutomationCommandQueue.Add(EAutomationCommand::RunCheckpointTests);
+					TestsRun = AutomationController->GetCheckpointFileContents();
+					AutomationController->CleanUpCheckpointFile();
+				}
+				else if (FParse::Command(&TempCmd, TEXT("SetMinimumPriority")))
+				{
+					StringCommand = TempCmd;
+					Ar.Logf(TEXT("Setting minimum priority of cases to run to: %s"), *StringCommand);
+					if (StringCommand.Contains(TEXT("Low")))
+					{
+						AutomationController->SetRequestedTestFlags(EAutomationTestFlags::PriorityMask);
+					}
+					else if (StringCommand.Contains(TEXT("Medium")))
+					{
+						AutomationController->SetRequestedTestFlags(EAutomationTestFlags::MediumPriority);
+					}
+					else if (StringCommand.Contains(TEXT("High")))
+					{
+						AutomationController->SetRequestedTestFlags(EAutomationTestFlags::HighPriorityAndAbove);
+					}
+					else if (StringCommand.Contains(TEXT("Critical")))
+					{
+						AutomationController->SetRequestedTestFlags(EAutomationTestFlags::ClientContext);
+					}
+					else if (StringCommand.Contains(TEXT("None")))
+					{
+						AutomationController->SetRequestedTestFlags(0);
+					}
+					else
+					{
+						Ar.Logf(TEXT("%s is not a valid priority!\nValid priorities are Critical, High, Medium, Low, None"), *StringCommand);
+					}
+				}
+				else if (FParse::Command(&TempCmd, TEXT("SetPriority")))
+				{
+					StringCommand = TempCmd;
+					Ar.Logf(TEXT("Setting explicit priority of cases to run to: %s"), *StringCommand);
+					if (StringCommand.Contains(TEXT("Low")))
+					{
+						AutomationController->SetRequestedTestFlags(EAutomationTestFlags::LowPriority);
+					}
+					else if (StringCommand.Contains(TEXT("Medium")))
+					{
+						AutomationController->SetRequestedTestFlags(EAutomationTestFlags::MediumPriority);
+					}
+					else if (StringCommand.Contains(TEXT("High")))
+					{
+						AutomationController->SetRequestedTestFlags(EAutomationTestFlags::HighPriority);
+					}
+					else if (StringCommand.Contains(TEXT("Critical")))
+					{
+						AutomationController->SetRequestedTestFlags(EAutomationTestFlags::CriticalPriority);
+					}
+					else if (StringCommand.Contains(TEXT("None")))
+					{
+						AutomationController->SetRequestedTestFlags(0);
+					}
+
+					else
+					{
+						Ar.Logf(TEXT("%s is not a valid priority!\nValid priorities are Critical, High, Medium, Low, None"), *StringCommand);
+					}
+				}
 				else if (FParse::Command(&TempCmd, TEXT("RunFilter")))
 				{
 					FlagToUse = TempCmd;
@@ -403,6 +516,9 @@ private:
 	/** The current state of the automation process */
 	EAutomationTestState::Type AutomationTestState;
 
+	/** The priority flags we would like to run */
+	EAutomationTestFlags::Type AutomationPriority;
+
 	/** What work was requested */
 	TArray<EAutomationCommand::Type> AutomationCommandQueue;
 
@@ -429,6 +545,13 @@ private:
 
 	//Dictionary that maps flag names to flag values.
 	TMap<FString, int32> FilterMaps;
+
+	//Test pass checkpoint backup file.
+	FArchive* CheckpointFile;
+
+	FString CheckpointCommand;
+
+	TArray<FString> TestsRun;
 };
 
 static FAutomationExecCmd AutomationExecCmd;
diff --git a/Engine/Source/Developer/AutomationController/Private/AutomationControllerManager.h b/Engine/Source/Developer/AutomationController/Private/AutomationControllerManager.h
--- a/Engine/Source/Developer/AutomationController/Private/AutomationControllerManager.h
+++ b/Engine/Source/Developer/AutomationController/Private/AutomationControllerManager.h
@@ -274,6 +274,17 @@ public:
 	virtual const bool IsTrackingHistory() const override;
 	virtual const int32 GetNumberHistoryItemsTracking() const override;
 
+	// Checkpoint logic
+	virtual TArray<FString> GetCheckpointFileContents() override;
+
+	virtual FArchive* GetCheckpointFileForWrite() override;
+
+	virtual void CleanUpCheckpointFile() override;
+
+	virtual void WriteLoadedCheckpointDataToFile() override;
+
+	virtual void WriteLineToCheckpointFile(FString LineToWrite) override;
+
 protected:
 
 	/**
@@ -498,6 +509,14 @@ private:
 	/** The report folder override path that may have been provided over the commandline, -ReportOutputPath="" */
 	FString ReportOutputPathOverride;
 
+	// Checkpoint variables
+	//Test pass checkpoint backup file.
+	FArchive* CheckpointFile;
+
+	FString CheckpointCommand;
+
+	TArray<FString> TestsRun;
+
 private:
 
 	/** Holds a delegate that is invoked when the controller shuts down. */
diff --git a/Engine/Source/Developer/AutomationController/Private/AutomationControllerManger.cpp b/Engine/Source/Developer/AutomationController/Private/AutomationControllerManger.cpp
--- a/Engine/Source/Developer/AutomationController/Private/AutomationControllerManger.cpp
+++ b/Engine/Source/Developer/AutomationController/Private/AutomationControllerManger.cpp
@@ -35,6 +35,7 @@ namespace AutomationControllerConstants
 FAutomationControllerManager::FAutomationControllerManager()
 {
 	FParse::Value(FCommandLine::Get(), TEXT("ReportOutputPath="), ReportOutputPathOverride, false);
+	CheckpointFile = nullptr;
 }
 
 void FAutomationControllerManager::RequestAvailableWorkers(const FGuid& SessionId)
@@ -552,6 +553,11 @@ void FAutomationControllerManager::ExecuteNextTask( int32 ClusterIndex, OUT bool
 							GLog->Logf(ELogVerbosity::Display, TEXT("Running Automation: '%s' (Class Name: '%s')"), *TestsRunThisPass[AddressIndex]->GetFullTestPath(), *TestsRunThisPass[AddressIndex]->GetCommand());
 							TestResults.State = EAutomationState::InProcess;
 
+							if (CheckpointFile)
+							{
+								WriteLineToCheckpointFile(NextTest->GetFullTestPath());
+							}
+
 							TestResults.GameInstance = DeviceClusterManager.GetClusterDeviceName(ClusterIndex, DeviceIndex);
 							NextTest->SetResults(ClusterIndex, CurrentTestPass, TestResults);
 							NextTest->ResetNetworkCommandResponses();
@@ -733,6 +739,7 @@ void FAutomationControllerManager::ProcessResults()
 
 	// Then clean our array for the next pass.
 	OurPassResults.ClearAllEntries();
+	CleanUpCheckpointFile();
 
 	SetControllerStatus(EAutomationControllerModuleState::Ready);
 }
@@ -1138,3 +1145,62 @@ const int32 FAutomationControllerManager::GetNumberHistoryItemsTracking() const
 {
 	return NumberOfHistoryItemsTracked;
 }
+
+TArray<FString> FAutomationControllerManager::GetCheckpointFileContents()
+{
+	TestsRun.Empty();
+	FString CheckpointFileName = FString::Printf(TEXT("%sautomationcheckpoint.txt"), *FPaths::AutomationDir());
+	if (IFileManager::Get().FileExists(*CheckpointFileName))
+	{
+		FString FileData;
+		FFileHelper::LoadFileToString(FileData, *CheckpointFileName);
+		FileData.ParseIntoArrayLines(TestsRun);
+		
+	}
+	return TestsRun;
+}
+
+FArchive* FAutomationControllerManager::GetCheckpointFileForWrite()
+{
+	if (!CheckpointFile)
+	{
+		FString CheckpointFileName = FString::Printf(TEXT("%sautomationcheckpoint.txt"), *FPaths::AutomationDir());
+		CheckpointFile = IFileManager::Get().CreateFileWriter(*CheckpointFileName, 8);
+	}
+	return CheckpointFile;
+}
+
+void FAutomationControllerManager::CleanUpCheckpointFile()
+{
+	if (CheckpointFile)
+	{
+		CheckpointFile->Close();
+		CheckpointFile = nullptr;
+	}
+	FString CheckpointFileName = FString::Printf(TEXT("%sautomationcheckpoint.txt"), *FPaths::AutomationDir());
+	if (IFileManager::Get().FileExists(*CheckpointFileName))
+	{
+		IFileManager::Get().Delete(*CheckpointFileName);
+	}
+}
+
+void FAutomationControllerManager::WriteLoadedCheckpointDataToFile()
+{
+	GetCheckpointFileForWrite();
+	for (int i = 0; i < TestsRun.Num(); i++)
+	{
+		FString LineToWrite = FString::Printf(TEXT("%s\n"), *TestsRun[i]);
+		CheckpointFile->Serialize(TCHAR_TO_ANSI(*LineToWrite), LineToWrite.Len());
+	}
+}
+
+void FAutomationControllerManager::WriteLineToCheckpointFile(FString StringToWrite)
+{
+	GetCheckpointFileForWrite();
+	if (CheckpointFile)
+	{
+		FString LineToWrite = FString::Printf(TEXT("%s\n"), *StringToWrite);
+		CheckpointFile->Serialize(TCHAR_TO_ANSI(*LineToWrite), LineToWrite.Len());
+		CheckpointFile->Flush();
+	}
+}
\ No newline at end of file
diff --git a/Engine/Source/Developer/AutomationController/Public/Interfaces/IAutomationControllerManager.h b/Engine/Source/Developer/AutomationController/Public/Interfaces/IAutomationControllerManager.h
--- a/Engine/Source/Developer/AutomationController/Public/Interfaces/IAutomationControllerManager.h
+++ b/Engine/Source/Developer/AutomationController/Public/Interfaces/IAutomationControllerManager.h
@@ -395,6 +395,16 @@ public:
 	 */
 	virtual const int32 GetNumberHistoryItemsTracking() const = 0;
 
+	virtual TArray<FString> GetCheckpointFileContents() = 0;
+
+	virtual FArchive* GetCheckpointFileForWrite() = 0;
+
+	virtual void CleanUpCheckpointFile() = 0;
+
+	virtual void WriteLoadedCheckpointDataToFile() = 0;
+
+	virtual void WriteLineToCheckpointFile(FString LineToWrite) = 0;
+
 public:
 
 	/**
diff --git a/Engine/Source/Developer/LogVisualizer/Private/VisualLoggerRenderingActor.cpp b/Engine/Source/Developer/LogVisualizer/Private/VisualLoggerRenderingActor.cpp
--- a/Engine/Source/Developer/LogVisualizer/Private/VisualLoggerRenderingActor.cpp
+++ b/Engine/Source/Developer/LogVisualizer/Private/VisualLoggerRenderingActor.cpp
@@ -401,14 +401,26 @@ void AVisualLoggerRenderingActor::GetDebugShapes(const FVisualLogDevice::FVisual
 		{
 			const float Radius = float(ElementToDraw->Radius);
 			const bool bDrawLabel = ElementToDraw->Description.IsEmpty() == false;
-			for (int32 Index = 0; Index < ElementToDraw->Points.Num(); ++Index)
+			if (ElementToDraw->Points.Num() == 1)
 			{
-				const FVector& Point = ElementToDraw->Points[Index];
+				const FVector& Point = ElementToDraw->Points[0];
 				DebugShapes.Points.Add(FDebugRenderSceneProxy::FSphere(Radius, Point, Color));
 				if (bDrawLabel)
 				{
-					const FString PrintString = FString::Printf(TEXT("%s_%d"), *ElementToDraw->Description, Index);
-					DebugShapes.Texts.Add(FDebugRenderSceneProxy::FText3d(PrintString, Point, Color));
+					DebugShapes.Texts.Add(FDebugRenderSceneProxy::FText3d(ElementToDraw->Description, Point, Color));
+				}
+			}
+			else
+			{
+				for (int32 Index = 0; Index < ElementToDraw->Points.Num(); ++Index)
+				{
+					const FVector& Point = ElementToDraw->Points[Index];
+					DebugShapes.Points.Add(FDebugRenderSceneProxy::FSphere(Radius, Point, Color));
+					if (bDrawLabel)
+					{
+						const FString PrintString = FString::Printf(TEXT("%s_%d"), *ElementToDraw->Description, Index);
+						DebugShapes.Texts.Add(FDebugRenderSceneProxy::FText3d(PrintString, Point, Color));
+					}
 				}
 			}
 		}
@@ -466,16 +478,30 @@ void AVisualLoggerRenderingActor::GetDebugShapes(const FVisualLogDevice::FVisual
 			const float Thickness = float(ElementToDraw->Thicknes);
 			const bool bDrawLabel = ElementToDraw->Description.IsEmpty() == false && ElementToDraw->Points.Num() > 2;
 			const FVector* Location = ElementToDraw->Points.GetData();
-			for (int32 Index = 0; Index + 1 < ElementToDraw->Points.Num(); Index += 2, Location += 2)
+			
+			if (ElementToDraw->Points.Num() == 1)
 			{
 				DebugShapes.Lines.Add(FDebugRenderSceneProxy::FDebugLine(*Location, *(Location + 1), Color, Thickness));
 
 				if (bDrawLabel)
 				{
-					const FString PrintString = FString::Printf(TEXT("%s_%d"), *ElementToDraw->Description, Index);
-					DebugShapes.Texts.Add(FDebugRenderSceneProxy::FText3d(PrintString, (*Location + (*(Location + 1) - *Location) / 2), Color));
+					DebugShapes.Texts.Add(FDebugRenderSceneProxy::FText3d(ElementToDraw->Description, (*Location + (*(Location + 1) - *Location) / 2), Color));
+				}
+			}
+			else
+			{
+				for (int32 Index = 0; Index + 1 < ElementToDraw->Points.Num(); Index += 2, Location += 2)
+				{
+					DebugShapes.Lines.Add(FDebugRenderSceneProxy::FDebugLine(*Location, *(Location + 1), Color, Thickness));
+
+					if (bDrawLabel)
+					{
+						const FString PrintString = FString::Printf(TEXT("%s_%d"), *ElementToDraw->Description, Index);
+						DebugShapes.Texts.Add(FDebugRenderSceneProxy::FText3d(PrintString, (*Location + (*(Location + 1) - *Location) / 2), Color));
+					}
 				}
 			}
+
 			if (ElementToDraw->Description.IsEmpty() == false)
 			{
 				DebugShapes.Texts.Add(FDebugRenderSceneProxy::FText3d(ElementToDraw->Description, ElementToDraw->Points[0] + (ElementToDraw->Points[1] - ElementToDraw->Points[0]) / 2, Color));
@@ -499,17 +525,32 @@ void AVisualLoggerRenderingActor::GetDebugShapes(const FVisualLogDevice::FVisual
 			const float Thickness = float(ElementToDraw->Thicknes);
 			const bool bDrawLabel = ElementToDraw->Description.IsEmpty() == false && ElementToDraw->Points.Num() > 2;
 			const FVector* BoxExtent = ElementToDraw->Points.GetData();
-			for (int32 Index = 0; Index + 1 < ElementToDraw->Points.Num(); Index += 2, BoxExtent += 2)
+			
+			if (ElementToDraw->Points.Num() == 2)
 			{
 				const FBox Box = FBox(*BoxExtent, *(BoxExtent + 1));
 				DebugShapes.Boxes.Add(FDebugRenderSceneProxy::FDebugBox(Box, Color, FTransform(ElementToDraw->TransformationMatrix)));
 
 				if (bDrawLabel)
 				{
-					const FString PrintString = FString::Printf(TEXT("%s_%d"), *ElementToDraw->Description, Index);
-					DebugShapes.Texts.Add(FDebugRenderSceneProxy::FText3d(PrintString, Box.GetCenter(), Color));
+					DebugShapes.Texts.Add(FDebugRenderSceneProxy::FText3d(ElementToDraw->Description, Box.GetCenter(), Color));
 				}
 			}
+			else
+			{
+				for (int32 Index = 0; Index + 1 < ElementToDraw->Points.Num(); Index += 2, BoxExtent += 2)
+				{
+					const FBox Box = FBox(*BoxExtent, *(BoxExtent + 1));
+					DebugShapes.Boxes.Add(FDebugRenderSceneProxy::FDebugBox(Box, Color, FTransform(ElementToDraw->TransformationMatrix)));
+
+					if (bDrawLabel)
+					{
+						const FString PrintString = FString::Printf(TEXT("%s_%d"), *ElementToDraw->Description, Index);
+						DebugShapes.Texts.Add(FDebugRenderSceneProxy::FText3d(PrintString, Box.GetCenter(), Color));
+					}
+				}
+			}
+
 			if (ElementToDraw->Description.IsEmpty() == false)
 			{
 				DebugShapes.Texts.Add(FDebugRenderSceneProxy::FText3d(ElementToDraw->Description, ElementToDraw->Points[0] + (ElementToDraw->Points[1] - ElementToDraw->Points[0]) / 2, Color));
diff --git a/Engine/Source/Editor/AnimationBlueprintEditor/Private/AnimationBlueprintEditor.cpp b/Engine/Source/Editor/AnimationBlueprintEditor/Private/AnimationBlueprintEditor.cpp
--- a/Engine/Source/Editor/AnimationBlueprintEditor/Private/AnimationBlueprintEditor.cpp
+++ b/Engine/Source/Editor/AnimationBlueprintEditor/Private/AnimationBlueprintEditor.cpp
@@ -59,6 +59,9 @@
 #include "PersonaCommonCommands.h"
 #include "AnimGraphCommands.h"
 
+#include "AnimGraphNode_AimOffsetLookAt.h"
+#include "AnimGraphNode_RotationOffsetBlendSpace.h"
+
 #define LOCTEXT_NAMESPACE "AnimationBlueprintEditor"
 
 const FName AnimationBlueprintEditorAppName(TEXT("AnimationBlueprintEditorApp"));
@@ -855,6 +858,131 @@ void FAnimationBlueprintEditor::OnConvertToPoseByName()
 	}
 }
 
+void FAnimationBlueprintEditor::OnConvertToAimOffsetLookAt()
+{
+	FGraphPanelSelectionSet SelectedNodes = GetSelectedNodes();
+
+	if (SelectedNodes.Num() > 0)
+	{
+		for (auto NodeIter = SelectedNodes.CreateIterator(); NodeIter; ++NodeIter)
+		{
+			UAnimGraphNode_RotationOffsetBlendSpace* OldNode = Cast<UAnimGraphNode_RotationOffsetBlendSpace>(*NodeIter);
+
+			// see if sequence player
+			if (OldNode && OldNode->Node.BlendSpace)
+			{
+				//const FScopedTransaction Transaction( LOCTEXT("ConvertToSequenceEvaluator", "Convert to Single Frame Animation") );
+
+				// convert to sequence evaluator
+				UEdGraph* TargetGraph = OldNode->GetGraph();
+				// create new evaluator
+				FGraphNodeCreator<UAnimGraphNode_AimOffsetLookAt> NodeCreator(*TargetGraph);
+				UAnimGraphNode_AimOffsetLookAt* NewNode = NodeCreator.CreateNode();
+				NewNode->Node.BlendSpace = OldNode->Node.BlendSpace;
+				NodeCreator.Finalize();
+
+				// get default data from old node to new node
+				FEdGraphUtilities::CopyCommonState(OldNode, NewNode);
+
+				UEdGraphPin* OldPosePin = OldNode->FindPin(TEXT("Pose"));
+				UEdGraphPin* NewPosePin = NewNode->FindPin(TEXT("Pose"));
+
+				if (ensure(OldPosePin && NewPosePin))
+				{
+					NewPosePin->CopyPersistentDataFromOldPin(*OldPosePin);
+				}
+
+				OldPosePin = OldNode->FindPin(TEXT("BasePose"));
+				NewPosePin = NewNode->FindPin(TEXT("BasePose"));
+
+				if (ensure(OldPosePin && NewPosePin))
+				{
+					NewPosePin->CopyPersistentDataFromOldPin(*OldPosePin);
+				}
+
+				// remove from selection and from graph
+				NodeIter.RemoveCurrent();
+				TargetGraph->RemoveNode(OldNode);
+
+				NewNode->Modify();
+			}
+		}
+
+		// @todo fixme: below code doesn't work
+		// because of SetAndCenterObject kicks in after new node is added
+		// will need to disable that first
+		TSharedPtr<SGraphEditor> FocusedGraphEd = FocusedGraphEdPtr.Pin();
+
+		// Update the graph so that the node will be refreshed
+		FocusedGraphEd->NotifyGraphChanged();
+		// It's possible to leave invalid objects in the selection set if they get GC'd, so clear it out
+		FocusedGraphEd->ClearSelectionSet();
+
+		FBlueprintEditorUtils::MarkBlueprintAsStructurallyModified(GetAnimBlueprint());
+	}
+}
+
+void FAnimationBlueprintEditor::OnConvertToAimOffsetSimple()
+{
+	FGraphPanelSelectionSet SelectedNodes = GetSelectedNodes();
+	if (SelectedNodes.Num() > 0)
+	{
+		for (auto NodeIter = SelectedNodes.CreateIterator(); NodeIter; ++NodeIter)
+		{
+			UAnimGraphNode_AimOffsetLookAt* OldNode = Cast<UAnimGraphNode_AimOffsetLookAt>(*NodeIter);
+
+			// see if sequence player
+			if (OldNode && OldNode->Node.BlendSpace)
+			{
+				//const FScopedTransaction Transaction( LOCTEXT("ConvertToSequenceEvaluator", "Convert to Single Frame Animation") );
+				// convert to sequence player
+				UEdGraph* TargetGraph = OldNode->GetGraph();
+				// create new player
+				FGraphNodeCreator<UAnimGraphNode_RotationOffsetBlendSpace> NodeCreator(*TargetGraph);
+				UAnimGraphNode_RotationOffsetBlendSpace* NewNode = NodeCreator.CreateNode();
+				NewNode->Node.BlendSpace = OldNode->Node.BlendSpace;
+				NodeCreator.Finalize();
+
+				// get default data from old node to new node
+				FEdGraphUtilities::CopyCommonState(OldNode, NewNode);
+
+				UEdGraphPin* OldPosePin = OldNode->FindPin(TEXT("Pose"));
+				UEdGraphPin* NewPosePin = NewNode->FindPin(TEXT("Pose"));
+
+				if (ensure(OldPosePin && NewPosePin))
+				{
+					NewPosePin->CopyPersistentDataFromOldPin(*OldPosePin);
+				}
+
+				OldPosePin = OldNode->FindPin(TEXT("BasePose"));
+				NewPosePin = NewNode->FindPin(TEXT("BasePose"));
+
+				if (ensure(OldPosePin && NewPosePin))
+				{
+					NewPosePin->CopyPersistentDataFromOldPin(*OldPosePin);
+				}
+
+				// remove from selection and from graph
+				NodeIter.RemoveCurrent();
+				TargetGraph->RemoveNode(OldNode);
+
+				NewNode->Modify();
+			}
+		}
+
+		// @todo fixme: below code doesn't work
+		// because of SetAndCenterObject kicks in after new node is added
+		// will need to disable that first
+		TSharedPtr<SGraphEditor> FocusedGraphEd = FocusedGraphEdPtr.Pin();
+		// Update the graph so that the node will be refreshed
+		FocusedGraphEd->NotifyGraphChanged();
+		// It's possible to leave invalid objects in the selection set if they get GC'd, so clear it out
+		FocusedGraphEd->ClearSelectionSet();
+
+		FBlueprintEditorUtils::MarkBlueprintAsStructurallyModified(GetAnimBlueprint());
+	}
+}
+
 void FAnimationBlueprintEditor::OnOpenRelatedAsset()
 {
 	FGraphPanelSelectionSet SelectedNodes = GetSelectedNodes();
diff --git a/Engine/Source/Editor/AnimationBlueprintEditor/Private/AnimationBlueprintEditor.h b/Engine/Source/Editor/AnimationBlueprintEditor/Private/AnimationBlueprintEditor.h
--- a/Engine/Source/Editor/AnimationBlueprintEditor/Private/AnimationBlueprintEditor.h
+++ b/Engine/Source/Editor/AnimationBlueprintEditor/Private/AnimationBlueprintEditor.h
@@ -166,6 +166,8 @@ protected:
 	virtual void OnConvertToBlendSpacePlayer() override;
 	virtual void OnConvertToPoseBlender() override;
 	virtual void OnConvertToPoseByName() override;
+	virtual void OnConvertToAimOffsetLookAt() override;
+	virtual void OnConvertToAimOffsetSimple() override;
 	virtual bool IsInAScriptingMode() const override { return true; }
 	virtual void OnOpenRelatedAsset() override;
 	virtual void GetCustomDebugObjects(TArray<FCustomDebugObject>& DebugList) const override;
diff --git a/Engine/Source/Editor/Persona/Private/PersonaMeshDetails.cpp b/Engine/Source/Editor/Persona/Private/PersonaMeshDetails.cpp
--- a/Engine/Source/Editor/Persona/Private/PersonaMeshDetails.cpp
+++ b/Engine/Source/Editor/Persona/Private/PersonaMeshDetails.cpp
@@ -75,14 +75,6 @@ static bool IsAutoMeshReductionAvailable()
 	return bAutoMeshReductionAvailable;
 }
 
-/** Returns true if Skincache and -shaders are enabled */
-static bool IsSkinCacheAvailable()
-{
-	extern ENGINE_API int32 GEnableGPUSkinCacheShaders;
-	static bool bSkinCacheAvailable = !!GEnableGPUSkinCacheShaders;
-	return bSkinCacheAvailable;
-}
-
 /**
 * FMeshReductionSettings
 */
@@ -1095,6 +1087,7 @@ void FPersonaMeshDetails::AddLODLevelCategories(IDetailLayoutBuilder& DetailLayo
 
 				MaterialListDelegates.OnGetMaterials.BindSP(this, &FPersonaMeshDetails::OnGetMaterialsForArray, 0);
 				MaterialListDelegates.OnMaterialChanged.BindSP(this, &FPersonaMeshDetails::OnMaterialArrayChanged, 0);
+				MaterialListDelegates.OnGenerateCustomNameWidgets.BindSP(this, &FPersonaMeshDetails::OnGenerateCustomNameWidgetsForMaterialArray);
 				MaterialListDelegates.OnGenerateCustomMaterialWidgets.BindSP(this, &FPersonaMeshDetails::OnGenerateCustomMaterialWidgetsForMaterialArray, 0);
 				MaterialListDelegates.OnMaterialListDirty.BindSP(this, &FPersonaMeshDetails::OnMaterialListDirty);
 
@@ -1952,6 +1945,39 @@ void FPersonaMeshDetails::OnMaterialNameChanged(const FText& InValue, int32 Mate
 	}
 }
 
+TSharedRef<SWidget> FPersonaMeshDetails::OnGenerateCustomNameWidgetsForMaterialArray(UMaterialInterface* Material, int32 MaterialIndex)
+{
+	return SNew(SVerticalBox)
+		+ SVerticalBox::Slot()
+		.AutoHeight()
+		[
+			SNew(SCheckBox)
+			.IsChecked(this, &FPersonaMeshDetails::IsMaterialSelected, MaterialIndex)
+			.OnCheckStateChanged(this, &FPersonaMeshDetails::OnMaterialSelectedChanged, MaterialIndex)
+			.ToolTipText(LOCTEXT("Highlight_CustomMaterialName_ToolTip", "Highlights this material in the viewport"))
+			[
+				SNew(STextBlock)
+				.Font(IDetailLayoutBuilder::GetDetailFont())
+				.ColorAndOpacity(FLinearColor(0.4f, 0.4f, 0.4f, 1.0f))
+				.Text(LOCTEXT("Highlight", "Highlight"))
+			]
+		]
+		+ SVerticalBox::Slot()
+		.AutoHeight()
+		.Padding(0, 2, 0, 0)
+		[
+			SNew(SCheckBox)
+			.IsChecked(this, &FPersonaMeshDetails::IsIsolateMaterialEnabled, MaterialIndex)
+			.OnCheckStateChanged(this, &FPersonaMeshDetails::OnMaterialIsolatedChanged, MaterialIndex)
+			.ToolTipText(LOCTEXT("Isolate_CustomMaterialName_ToolTip", "Isolates this material in the viewport"))
+			[
+				SNew(STextBlock)
+				.Font(IDetailLayoutBuilder::GetDetailFont())
+				.ColorAndOpacity(FLinearColor(0.4f, 0.4f, 0.4f, 1.0f))
+				.Text(LOCTEXT("Isolate", "Isolate"))
+			]
+		];
+}
 
 TSharedRef<SWidget> FPersonaMeshDetails::OnGenerateCustomMaterialWidgetsForMaterialArray(UMaterialInterface* Material, int32 MaterialIndex, int32 LODIndex)
 {
@@ -2115,8 +2141,8 @@ TSharedRef<SWidget> FPersonaMeshDetails::OnGenerateCustomNameWidgetsForSection(i
 		.AutoHeight()
 		[
 			SNew(SCheckBox)
-			.IsChecked(this, &FPersonaMeshDetails::IsSectionSelected, SectionIndex)
-			.OnCheckStateChanged(this, &FPersonaMeshDetails::OnSectionSelectedChanged, SectionIndex)
+			.IsChecked(this, &FPersonaMeshDetails::IsSectionSelected, ((SectionIndex << 16) | LodIndex))
+			.OnCheckStateChanged(this, &FPersonaMeshDetails::OnSectionSelectedChanged, ((SectionIndex << 16) | LodIndex))
 			.ToolTipText(LOCTEXT("Highlight_ToolTip", "Highlights this section in the viewport"))
 			[
 				SNew(STextBlock)
@@ -2130,8 +2156,8 @@ TSharedRef<SWidget> FPersonaMeshDetails::OnGenerateCustomNameWidgetsForSection(i
 		.Padding(0, 2, 0, 0)
 			[
 			SNew(SCheckBox)
-			.IsChecked(this, &FPersonaMeshDetails::IsIsolateSectionEnabled, SectionIndex)
-			.OnCheckStateChanged(this, &FPersonaMeshDetails::OnSectionIsolatedChanged, SectionIndex)
+			.IsChecked(this, &FPersonaMeshDetails::IsIsolateSectionEnabled, ((SectionIndex << 16) | LodIndex))
+			.OnCheckStateChanged(this, &FPersonaMeshDetails::OnSectionIsolatedChanged, ((SectionIndex << 16) | LodIndex))
 			.ToolTipText(LOCTEXT("Isolate_ToolTip", "Isolates this section in the viewport"))
 				[
 				SNew(STextBlock)
@@ -2144,6 +2170,8 @@ TSharedRef<SWidget> FPersonaMeshDetails::OnGenerateCustomNameWidgetsForSection(i
 
 TSharedRef<SWidget> FPersonaMeshDetails::OnGenerateCustomSectionWidgetsForSection(int32 LODIndex, int32 SectionIndex)
 {
+	extern ENGINE_API bool IsGPUSkinCacheAvailable();
+
 	TSharedRef<SVerticalBox> SectionWidget
 		= SNew(SVerticalBox)
 
@@ -2164,7 +2192,7 @@ TSharedRef<SWidget> FPersonaMeshDetails::OnGenerateCustomSectionWidgetsForSectio
 		.Padding(0, 2, 0, 0)
 		[
 			SNew(SCheckBox)
-			.IsEnabled(IsSkinCacheAvailable())
+			.IsEnabled(IsGPUSkinCacheAvailable())
 			.IsChecked(this, &FPersonaMeshDetails::IsSectionRecomputeTangentEnabled, LODIndex, SectionIndex)
 			.OnCheckStateChanged(this, &FPersonaMeshDetails::OnSectionRecomputeTangentChanged, LODIndex, SectionIndex)
 			[
@@ -2213,20 +2241,112 @@ TSharedRef<SWidget> FPersonaMeshDetails::OnGenerateCustomSectionWidgetsForSectio
 	return SectionWidget;
 }
 
-ECheckBoxState FPersonaMeshDetails::IsSectionSelected(int32 SectionIndex) const
+bool FPersonaMeshDetails::GetFirstLodSectionUsingMaterial(int32 MaterialIndex, int32 &OutMaterialLODIndex, int32 &OutMaterialSectionIndex) const
+{
+	//Find the first section using this material index, it will be the one we use to know if we isolate or select
+	const UDebugSkelMeshComponent* MeshComponent = GetPersonaToolkit()->GetPreviewScene()->GetPreviewMeshComponent();
+	const USkeletalMesh* Mesh = GetPersonaToolkit()->GetMesh();
+	if (MeshComponent && Mesh)
+	{
+		check(Mesh->GetResourceForRendering());
+		for (int32 LodIndex = 0; LodIndex < Mesh->GetResourceForRendering()->LODModels.Num(); ++LodIndex)
+		{
+			if (OutMaterialLODIndex != INDEX_NONE)
+			{
+				break;
+			}
+			const FStaticLODModel& LODModel = Mesh->GetResourceForRendering()->LODModels[LodIndex];
+			for (int32 SectionIndex = 0; SectionIndex < LODModel.Sections.Num(); ++SectionIndex)
+			{
+				int32 LodSectionMaterialIndex = GetMaterialIndex(LodIndex, SectionIndex);
+				if (MaterialIndex == LodSectionMaterialIndex)
+				{
+					OutMaterialLODIndex = LodIndex;
+					OutMaterialSectionIndex = SectionIndex;
+					break;
+				}
+			}
+		}
+	}
+
+	return OutMaterialLODIndex != INDEX_NONE ? true : false;
+}
+
+ECheckBoxState FPersonaMeshDetails::IsMaterialSelected(int32 MaterialIndex) const
 {
+	int32 MaterialLODIndex = INDEX_NONE;
+	int32 MaterialSectionIndex = INDEX_NONE;
+	if (!GetFirstLodSectionUsingMaterial(MaterialIndex, MaterialLODIndex, MaterialSectionIndex))
+	{
+		//There is no section using this material
+		return ECheckBoxState::Unchecked;
+	}
+	return IsSectionSelected(((MaterialSectionIndex << 16) | MaterialLODIndex));
+}
+
+void FPersonaMeshDetails::OnMaterialSelectedChanged(ECheckBoxState NewState, int32 MaterialIndex)
+{
+	int32 MaterialLODIndex = INDEX_NONE;
+	int32 MaterialSectionIndex = INDEX_NONE;
+	if (!GetFirstLodSectionUsingMaterial(MaterialIndex, MaterialLODIndex, MaterialSectionIndex))
+	{
+		return;
+	}
+	OnSectionSelectedChanged(NewState, ((MaterialSectionIndex << 16) | MaterialLODIndex));
+}
+
+ECheckBoxState FPersonaMeshDetails::IsIsolateMaterialEnabled(int32 MaterialIndex) const
+{
+	int32 MaterialLODIndex = INDEX_NONE;
+	int32 MaterialSectionIndex = INDEX_NONE;
+	if (!GetFirstLodSectionUsingMaterial(MaterialIndex, MaterialLODIndex, MaterialSectionIndex))
+	{
+		//There is no section using this material
+		return ECheckBoxState::Unchecked;
+	}
+	return IsIsolateSectionEnabled(((MaterialSectionIndex << 16) | MaterialLODIndex));
+}
+
+void FPersonaMeshDetails::OnMaterialIsolatedChanged(ECheckBoxState NewState, int32 MaterialIndex)
+{
+	int32 MaterialLODIndex = INDEX_NONE;
+	int32 MaterialSectionIndex = INDEX_NONE;
+	if (!GetFirstLodSectionUsingMaterial(MaterialIndex, MaterialLODIndex, MaterialSectionIndex))
+	{
+		return;
+	}
+	OnSectionIsolatedChanged(NewState, ((MaterialSectionIndex << 16) | MaterialLODIndex));
+}
+
+
+ECheckBoxState FPersonaMeshDetails::IsSectionSelected(int32 SectionLODIndex) const
+{
+	int32 SectionIndex = SectionLODIndex >> 16;
+	int32 LODIdx = (SectionLODIndex & 0x0000FFFF);
+
 	ECheckBoxState State = ECheckBoxState::Unchecked;
 	const USkeletalMesh* Mesh = GetPersonaToolkit()->GetMesh();
 	if (Mesh)
 	{
-		State = Mesh->SelectedEditorSection == SectionIndex ? ECheckBoxState::Checked : ECheckBoxState::Unchecked;
+		check(Mesh->GetResourceForRendering());
+		if (Mesh->GetResourceForRendering()->LODModels.IsValidIndex(LODIdx))
+		{
+			const FStaticLODModel& LODModel = Mesh->GetResourceForRendering()->LODModels[LODIdx];
+			if (LODModel.Sections.IsValidIndex(SectionIndex))
+			{
+				int32 MaterialIndex = GetMaterialIndex(LODIdx, SectionIndex);
+				State = Mesh->SelectedEditorSection == MaterialIndex ? ECheckBoxState::Checked : ECheckBoxState::Unchecked;
+			}
+		}
 	}
-
 	return State;
 }
 
-void FPersonaMeshDetails::OnSectionSelectedChanged(ECheckBoxState NewState, int32 SectionIndex)
+void FPersonaMeshDetails::OnSectionSelectedChanged(ECheckBoxState NewState, int32 SectionLODIndex)
 {
+	int32 SectionIndex = SectionLODIndex >> 16;
+	int32 LODIdx = (SectionLODIndex & 0x0000FFFF);
+
 	USkeletalMesh* Mesh = GetPersonaToolkit()->GetMesh();
 
 	// Currently assumes that we only ever have one preview mesh in Persona.
@@ -2234,43 +2354,76 @@ void FPersonaMeshDetails::OnSectionSelectedChanged(ECheckBoxState NewState, int3
 
 	if (Mesh && MeshComponent)
 	{
-		if (NewState == ECheckBoxState::Checked)
+		check(Mesh->GetResourceForRendering());
+		if (Mesh->GetResourceForRendering()->LODModels.IsValidIndex(LODIdx))
 		{
-			Mesh->SelectedEditorSection = SectionIndex;
-			if (MeshComponent->SectionIndexPreview != SectionIndex)
+			const FStaticLODModel& LODModel = Mesh->GetResourceForRendering()->LODModels[LODIdx];
+			if (LODModel.Sections.IsValidIndex(SectionIndex))
 			{
-				// Unhide all mesh sections
-				MeshComponent->SetSectionPreview(INDEX_NONE);
+				int32 MaterialIndex = GetMaterialIndex(LODIdx, SectionIndex);
+				if (NewState == ECheckBoxState::Checked)
+				{
+					Mesh->SelectedEditorSection = MaterialIndex;
+					if (MeshComponent->SectionIndexPreview != MaterialIndex)
+					{
+						// Unhide all mesh sections
+						MeshComponent->SetSectionPreview(INDEX_NONE);
+					}
+				}
+				else if (NewState == ECheckBoxState::Unchecked)
+				{
+					Mesh->SelectedEditorSection = INDEX_NONE;
+				}
+				GetPersonaToolkit()->GetPreviewScene()->InvalidateViews();
 			}
 		}
-		else if (NewState == ECheckBoxState::Unchecked)
-		{
-			Mesh->SelectedEditorSection = INDEX_NONE;
-		}
-		GetPersonaToolkit()->GetPreviewScene()->InvalidateViews();
 	}
 }
 
-ECheckBoxState FPersonaMeshDetails::IsIsolateSectionEnabled(int32 SectionIndex) const
+ECheckBoxState FPersonaMeshDetails::IsIsolateSectionEnabled(int32 SectionLODIndex) const
 {
+	int32 SectionIndex = SectionLODIndex >> 16;
+	int32 LODIdx = (SectionLODIndex & 0x0000FFFF);
 	ECheckBoxState State = ECheckBoxState::Unchecked;
 	const UDebugSkelMeshComponent* MeshComponent = GetPersonaToolkit()->GetPreviewScene()->GetPreviewMeshComponent();
-	if (MeshComponent)
+	const USkeletalMesh* Mesh = GetPersonaToolkit()->GetMesh();
+	if (MeshComponent && Mesh)
 	{
-		State = MeshComponent->SectionIndexPreview == SectionIndex ? ECheckBoxState::Checked : ECheckBoxState::Unchecked;
+		check(Mesh->GetResourceForRendering());
+		if (Mesh->GetResourceForRendering()->LODModels.IsValidIndex(LODIdx))
+		{
+			const FStaticLODModel& LODModel = Mesh->GetResourceForRendering()->LODModels[LODIdx];
+			if (LODModel.Sections.IsValidIndex(SectionIndex))
+			{
+				int32 MaterialIndex = GetMaterialIndex(LODIdx, SectionIndex);
+				State = MeshComponent->SectionIndexPreview == MaterialIndex ? ECheckBoxState::Checked : ECheckBoxState::Unchecked;
+			}
+		}
 	}
 	return State;
 }
 
-void FPersonaMeshDetails::OnSectionIsolatedChanged(ECheckBoxState NewState, int32 SectionIndex)
+void FPersonaMeshDetails::OnSectionIsolatedChanged(ECheckBoxState NewState, int32 SectionLODIndex)
 {
+	int32 SectionIndex = SectionLODIndex >> 16;
+	int32 LODIdx = (SectionLODIndex & 0x0000FFFF);
 	USkeletalMesh* Mesh = GetPersonaToolkit()->GetMesh();
 	UDebugSkelMeshComponent * MeshComponent = GetPersonaToolkit()->GetPreviewScene()->GetPreviewMeshComponent();
 	if (Mesh && MeshComponent)
 	{
 		if (NewState == ECheckBoxState::Checked)
 		{
-			MeshComponent->SetSectionPreview(SectionIndex);
+			int32 MaterialIndexToIsolate = INDEX_NONE;
+			check(Mesh->GetResourceForRendering());
+			if(Mesh->GetResourceForRendering()->LODModels.IsValidIndex(LODIdx))
+			{
+				const FStaticLODModel& LODModel = Mesh->GetResourceForRendering()->LODModels[LODIdx];
+				if (LODModel.Sections.IsValidIndex(SectionIndex))
+				{
+					MaterialIndexToIsolate = GetMaterialIndex(LODIdx, SectionIndex);
+				}
+			}
+			MeshComponent->SetSectionPreview(MaterialIndexToIsolate);
 			if (Mesh->SelectedEditorSection != SectionIndex)
 			{
 				Mesh->SelectedEditorSection = INDEX_NONE;
diff --git a/Engine/Source/Editor/Persona/Private/PersonaMeshDetails.h b/Engine/Source/Editor/Persona/Private/PersonaMeshDetails.h
--- a/Engine/Source/Editor/Persona/Private/PersonaMeshDetails.h
+++ b/Engine/Source/Editor/Persona/Private/PersonaMeshDetails.h
@@ -211,6 +211,14 @@ private:
 
 	FText GetOriginalImportMaterialNameText(int32 MaterialIndex)const;
 
+	/**
+	* Called by the material list widget on generating name side content
+	*
+	* @param Material		The material that is being displayed
+	* @param MaterialIndex	The index of the material slot
+	*/
+	TSharedRef<SWidget> OnGenerateCustomNameWidgetsForMaterialArray(UMaterialInterface* Material, int32 MaterialIndex);
+	
 	/**
 	* Called by the material list widget on generating each thumbnail widget
 	*
@@ -230,6 +238,36 @@ private:
 
 	FText GetFirstMaterialSlotUsedBySection(int32 MaterialIndex) const;
 
+	bool GetFirstLodSectionUsingMaterial(int32 MaterialIndex, int32 &OutMaterialLODIndex, int32 &OutMaterialSectionIndex) const;
+	/**
+	* Handler for check box display based on whether the material is highlighted
+	*
+	* @param MaterialIndex	The material index that is being selected
+	*/
+	ECheckBoxState IsMaterialSelected(int32 MaterialIndex) const;
+
+	/**
+	* Handler for changing highlight status on a material
+	*
+	* @param MaterialIndex	The material index that is being selected
+	*/
+	void OnMaterialSelectedChanged(ECheckBoxState NewState, int32 MaterialIndex);
+
+	/**
+	* Handler for check box display based on whether the material is isolated
+	*
+	* @param MaterialIndex	The material index that is being isolate
+	*/
+	ECheckBoxState IsIsolateMaterialEnabled(int32 MaterialIndex) const;
+
+	/**
+	* Handler for changing isolated status on a material
+	*
+	* @param MaterialIndex	The material index that is being isolate
+	*/
+	void OnMaterialIsolatedChanged(ECheckBoxState NewState, int32 MaterialIndex);
+
+
 	/**
 	 * Handler for check box display based on whether the material is highlighted
 	 *
diff --git a/Engine/Source/Editor/UMGEditor/Private/Hierarchy/SHierarchyView.h b/Engine/Source/Editor/UMGEditor/Private/Hierarchy/SHierarchyView.h
--- a/Engine/Source/Editor/UMGEditor/Private/Hierarchy/SHierarchyView.h
+++ b/Engine/Source/Editor/UMGEditor/Private/Hierarchy/SHierarchyView.h
@@ -9,6 +9,7 @@
 #include "WidgetBlueprintEditor.h"
 #include "Widgets/Views/STableViewBase.h"
 #include "Widgets/Views/STableRow.h"
+#include "Widgets/Views/STreeView.h"
 #include "Misc/TextFilter.h"
 #include "TreeFilterHandler.h"
 
diff --git a/Engine/Source/Editor/UnrealEd/Classes/CookOnTheSide/CookOnTheFlyServer.h b/Engine/Source/Editor/UnrealEd/Classes/CookOnTheSide/CookOnTheFlyServer.h
--- a/Engine/Source/Editor/UnrealEd/Classes/CookOnTheSide/CookOnTheFlyServer.h
+++ b/Engine/Source/Editor/UnrealEd/Classes/CookOnTheSide/CookOnTheFlyServer.h
@@ -26,22 +26,23 @@ enum class ESavePackageResult;
 
 enum class ECookInitializationFlags
 {
-	None =										0x0000, // No flags
-	Compressed =								0x0001, // will save compressed packages
-	Iterative =									0x0002, // use iterative cooking (previous cooks will not be cleaned unless detected out of date, experimental)
-	SkipEditorContent =							0x0004, // do not cook any content in the content\editor directory
-	Unversioned =								0x0008, // save the cooked packages without a version number
-	AutoTick =									0x0010, // enable ticking (only works in the editor)
-	AsyncSave =									0x0020, // save packages async
-	IncludeServerMaps =							0x0080, // should we include the server maps when cooking
-	UseSerializationForPackageDependencies =	0x0100, // should we use the serialization code path for generating package dependencies (old method will be deprecated)
-	BuildDDCInBackground =						0x0200, // build ddc content in background while the editor is running (only valid for modes which are in editor IsCookingInEditor())
-	GeneratedAssetRegistry =					0x0400, // have we generated asset registry yet
-	OutputVerboseCookerWarnings =				0x0800, // output additional cooker warnings about content issues
-	EnablePartialGC =							0x1000, // mark up with an object flag objects which are in packages which we are about to use or in the middle of using, this means we can gc more often but only gc stuff which we have finished with
-	TestCook =									0x2000, // test the cooker garbage collection process and cooking (cooker will never end just keep testing).
-	IterateOnHash =								0x4000, // when using iterative cooking use hashes of original files instead of timestamps
-	LogDebugInfo =								0x8000, // enables additional debug log information
+	None =										0x00000000, // No flags
+	Compressed =								0x00000001, // will save compressed packages
+	Iterative =									0x00000002, // use iterative cooking (previous cooks will not be cleaned unless detected out of date, experimental)
+	SkipEditorContent =							0x00000004, // do not cook any content in the content\editor directory
+	Unversioned =								0x00000008, // save the cooked packages without a version number
+	AutoTick =									0x00000010, // enable ticking (only works in the editor)
+	AsyncSave =									0x00000020, // save packages async
+	IncludeServerMaps =							0x00000080, // should we include the server maps when cooking
+	UseSerializationForPackageDependencies =	0x00000100, // should we use the serialization code path for generating package dependencies (old method will be deprecated)
+	BuildDDCInBackground =						0x00000200, // build ddc content in background while the editor is running (only valid for modes which are in editor IsCookingInEditor())
+	GeneratedAssetRegistry =					0x00000400, // have we generated asset registry yet
+	OutputVerboseCookerWarnings =				0x00000800, // output additional cooker warnings about content issues
+	EnablePartialGC =							0x00001000, // mark up with an object flag objects which are in packages which we are about to use or in the middle of using, this means we can gc more often but only gc stuff which we have finished with
+	TestCook =									0x00002000, // test the cooker garbage collection process and cooking (cooker will never end just keep testing).
+	IterateOnHash =								0x00004000, // when using iterative cooking use hashes of original files instead of timestamps
+	LogDebugInfo =								0x00008000, // enables additional debug log information
+	IterateSharedBuild =								0x00010000,
 };
 ENUM_CLASS_FLAGS(ECookInitializationFlags);
 
@@ -310,15 +311,53 @@ private:
 	struct FFilePlatformCookedPackage : public FFilePlatformRequest
 	{
 	private:
-		bool bSucceededSavePackage;
+		TArray<bool> bSucceededSavePackage; // one bool for each platform
 	public:
-		FFilePlatformCookedPackage(const FFilePlatformRequest& InFilePlatformRequest, bool bInSuccededSavePackage) : FFilePlatformRequest(InFilePlatformRequest.GetFilename(), InFilePlatformRequest.GetPlatformnames()), bSucceededSavePackage(bInSuccededSavePackage) { }
-		FFilePlatformCookedPackage(const FName& InFilename, const TArray<FName>& InPlatformname, bool bInSuccededSavePackage) : FFilePlatformRequest(InFilename, InPlatformname), bSucceededSavePackage(bInSuccededSavePackage) { }
-		FFilePlatformCookedPackage(const FName& InFilename, TArray<FName>&& InPlatformname, bool bInSuccededSavePackage) : FFilePlatformRequest(InFilename, MoveTemp(InPlatformname)), bSucceededSavePackage(bInSuccededSavePackage) { }
-		FFilePlatformCookedPackage(const FFilePlatformCookedPackage& InFilePlatformRequest) : FFilePlatformRequest(InFilePlatformRequest.Filename, InFilePlatformRequest.Platformnames), bSucceededSavePackage(InFilePlatformRequest.bSucceededSavePackage) { }
-		FFilePlatformCookedPackage(FFilePlatformCookedPackage&& InFilePlatformRequest) : FFilePlatformRequest(MoveTemp(InFilePlatformRequest.Filename), MoveTemp(InFilePlatformRequest.Platformnames)), bSucceededSavePackage(InFilePlatformRequest.bSucceededSavePackage) { }
+		FFilePlatformCookedPackage(const FFilePlatformRequest& InFilePlatformRequest, const TArray<bool>& bInSuccededSavePackage) : FFilePlatformRequest(InFilePlatformRequest.GetFilename(), InFilePlatformRequest.GetPlatformnames()), bSucceededSavePackage(bInSuccededSavePackage) { check(Platformnames.Num() == bSucceededSavePackage.Num()); }
+		FFilePlatformCookedPackage(const FName& InFilename, const TArray<FName>& InPlatformname) : FFilePlatformRequest(InFilename, InPlatformname)
+		{ 
+			// only use this constructor to short hand when packages fail
+			for ( int32 I = 0; I < InPlatformname.Num(); ++I )
+			{
+				bSucceededSavePackage.Add(false);
+			}
+			check(Platformnames.Num() == bSucceededSavePackage.Num());
+		}
+		FFilePlatformCookedPackage(const FName& InFilename, const TArray<FName>& InPlatformname, const TArray<bool>& bInSuccededSavePackage) : FFilePlatformRequest(InFilename, InPlatformname), bSucceededSavePackage(bInSuccededSavePackage) { check(Platformnames.Num() == bSucceededSavePackage.Num()); }
+		FFilePlatformCookedPackage(const FName& InFilename, TArray<FName>&& InPlatformname, const TArray<bool>&& bInSuccededSavePackage) : FFilePlatformRequest(InFilename, MoveTemp(InPlatformname)), bSucceededSavePackage(MoveTemp(bInSuccededSavePackage)) { check(Platformnames.Num() == bSucceededSavePackage.Num()); }
+		FFilePlatformCookedPackage(const FFilePlatformCookedPackage& InFilePlatformRequest) : FFilePlatformRequest(InFilePlatformRequest.Filename, InFilePlatformRequest.Platformnames), bSucceededSavePackage(InFilePlatformRequest.bSucceededSavePackage) { check(Platformnames.Num() == bSucceededSavePackage.Num());  }
+		FFilePlatformCookedPackage(FFilePlatformCookedPackage&& InFilePlatformRequest) : FFilePlatformRequest(MoveTemp(InFilePlatformRequest.Filename), MoveTemp(InFilePlatformRequest.Platformnames)), bSucceededSavePackage(InFilePlatformRequest.bSucceededSavePackage) { check(Platformnames.Num() == bSucceededSavePackage.Num()); }
+
+		inline const void AddPlatform( const FName& Platform, bool bSucceeded )
+		{
+			check(Platformnames.Num() == bSucceededSavePackage.Num());
+			check(Platform != NAME_None);
+			Platformnames.Add(Platform);
+			bSucceededSavePackage.Add(bSucceeded);
+			check( Platformnames.Num() == bSucceededSavePackage.Num() );
+		}
+
+		inline void RemovePlatform(const FName &Platform)
+		{
+			check(Platformnames.Num() == bSucceededSavePackage.Num());
+			int32 Index = Platformnames.IndexOfByKey(Platform);
+			if (Index != -1)
+			{
+				Platformnames.RemoveAt(Index);
+				bSucceededSavePackage.RemoveAt(Index);
+			}
+			check(Platformnames.Num() == bSucceededSavePackage.Num());
+		}
 
-		inline const bool HasSucceededSavePackage() const { return bSucceededSavePackage; }
+		inline const bool HasSucceededSavePackage(const FName& PlatformName) const 
+		{ 
+			int32 Index = Platformnames.IndexOfByKey(PlatformName);
+			if ( (Index != -1) && (Index < bSucceededSavePackage.Num()) )
+			{
+				return bSucceededSavePackage[Index]; 
+			}
+			return false;
+		}
 	};
 
 	/** Helper list of all files which have been cooked */
@@ -344,14 +383,15 @@ private:
 			check(Request.IsValid());
 
 			// see if it's already in the requests list
-			FFilePlatformRequest *ExistingRequest = FilesProcessed.Find(Request.GetFilename() );
+			FFilePlatformCookedPackage *ExistingRequest = FilesProcessed.Find(Request.GetFilename() );
 			
 			if ( ExistingRequest )
 			{
 				check( ExistingRequest->GetFilename() == Request.GetFilename() );
 				for ( const auto &Platform : Request.GetPlatformnames() )
 				{
-					ExistingRequest->AddPlatform(Platform);
+					const bool bSucceeded = Request.HasSucceededSavePackage(Platform);
+					ExistingRequest->AddPlatform(Platform, bSucceeded);
 				}
 			}
 			else
@@ -418,7 +458,16 @@ private:
 
 			if (bIncludeFailed == false)
 			{
-				if ( OurRequest->HasSucceededSavePackage() == false )
+				bool allFailed = true;
+				for ( const auto& PlatformName : PlatformNames )
+				{
+					if ( OurRequest->HasSucceededSavePackage(PlatformName) == false )
+					{
+						allFailed = false;
+						break;
+					}
+				}
+				if ( allFailed )
 				{
 					return false;
 				}
@@ -481,7 +530,7 @@ private:
 			{
 				if (CookedFile.Value.HasPlatform(PlatformName))
 				{
-					bool bHasSucceededSavePackage = CookedFile.Value.HasSucceededSavePackage();
+					bool bHasSucceededSavePackage = CookedFile.Value.HasSucceededSavePackage(PlatformName);
 					if ((bHasSucceededSavePackage && bGetSuccessfulCookedPackages) ||
 						((bHasSucceededSavePackage == false) && bGetFailedCookedPackages) )
 					CookedFiles.Add(CookedFile.Value.GetFilename());
@@ -832,6 +881,7 @@ private:
 	ECookInitializationFlags CookFlags;
 	TUniquePtr<class FSandboxPlatformFile> SandboxFile;
 	bool bIsInitializingSandbox; // stop recursion into callbacks when we are initializing sandbox
+	mutable bool bIgnoreMarkupPackageAlreadyLoaded; // avoid marking up packages as already loaded (want to put this around some functionality as we want to load packages fully some times)
 	bool bIsSavingPackage; // used to stop recursive mark package dirty functions
 
 	//////////////////////////////////////////////////////////////////////////
@@ -1026,6 +1076,12 @@ public:
 	 */
 	void ValidateCookByTheBookSettings() const;
 
+	/**
+	 * ValidateCookOnTheFlySettings
+	 * look at the initialization flags and other cooker settings make sure the programmer that thought of checking them are ok
+	 */
+	void ValidateCookOnTheFlySettings() const;
+
 	/**
 	 * Queue a cook by the book cancel (you might want to do this instead of calling cancel directly so that you don't have to be in the game thread when canceling
 	 */
@@ -1511,7 +1567,7 @@ private:
 	 *
 	 *	@return	ESavePackageResult::Success if packages was cooked
 	 */
-	ESavePackageResult SaveCookedPackage(UPackage* Package, uint32 SaveFlags, bool& bOutWasUpToDate);
+	void SaveCookedPackage(UPackage* Package, uint32 SaveFlags, bool& bOutWasUpToDate, TArray<ESavePackageResult>& SavePackageResults);
 	/**
 	 *	Cook (save) the given package
 	 *
@@ -1523,7 +1579,7 @@ private:
 	 *
 	 *	@return	ESavePackageResult::Success if packages was cooked
 	 */
-	ESavePackageResult SaveCookedPackage(UPackage* Package, uint32 SaveFlags, bool& bOutWasUpToDate, TArray<FName> &TargetPlatformNames);
+	void SaveCookedPackage(UPackage* Package, uint32 SaveFlags, bool& bOutWasUpToDate, TArray<FName> &TargetPlatformNames, TArray<ESavePackageResult>& SavePackageResults);
 
 
 	/**
@@ -1550,6 +1606,17 @@ private:
 	 */
 	void PopulateCookedPackagesFromDisk( const TArray<ITargetPlatform*>& Platforms );
 
+
+
+	/**
+	 * Searches the disk for all the cooked files in the sandbox path provided
+	 * Returns a map of the uncooked file path matches to the cooked file path for each package which exists
+	 *
+	 * @param UncookedpathToCookedPath out Map of the uncooked path matched with the cooked package which exists
+	 * @param SandboxPath path to search for cooked packages in
+	 */
+	void GetAllCookedFiles(TMap<FName, FName>& UncookedPathToCookedPath, const FString& SandboxPath);
+
 	/**
 	 * Verify the cooked package list hasn't got any packages which are out of date in it and remove any which are out of date deleting content from disk
 	 * 
diff --git a/Engine/Source/Editor/UnrealEd/Private/Commandlets/CookCommandlet.cpp b/Engine/Source/Editor/UnrealEd/Private/Commandlets/CookCommandlet.cpp
--- a/Engine/Source/Editor/UnrealEd/Private/Commandlets/CookCommandlet.cpp
+++ b/Engine/Source/Editor/UnrealEd/Private/Commandlets/CookCommandlet.cpp
@@ -1337,6 +1337,7 @@ bool UCookCommandlet::NewCook( const TArray<ITargetPlatform*>& Platforms, TArray
 	bool bTestCook = FParse::Param(*Params, TEXT("Testcook"));
 	CookFlags |= bTestCook ? ECookInitializationFlags::TestCook : ECookInitializationFlags::None;
 	CookFlags |= FParse::Param(*Params, TEXT("iteratehash")) ? ECookInitializationFlags::IterateOnHash : ECookInitializationFlags::None;
+	CookFlags |= FParse::Param(*Params, TEXT("iteratesharedbuild")) ? ECookInitializationFlags::IterateSharedBuild : ECookInitializationFlags::None;
 	CookFlags |= FParse::Param(*Params, TEXT("logdebuginfo")) ? ECookInitializationFlags::LogDebugInfo : ECookInitializationFlags::None;
 
 	// shared cooked build flags
diff --git a/Engine/Source/Editor/UnrealEd/Private/Commandlets/DerivedDataCacheCommandlet.cpp b/Engine/Source/Editor/UnrealEd/Private/Commandlets/DerivedDataCacheCommandlet.cpp
--- a/Engine/Source/Editor/UnrealEd/Private/Commandlets/DerivedDataCacheCommandlet.cpp
+++ b/Engine/Source/Editor/UnrealEd/Private/Commandlets/DerivedDataCacheCommandlet.cpp
@@ -201,8 +201,8 @@ int32 UDerivedDataCacheCommandlet::Main( const FString& Params )
 				}
 				if (bDoSubset)
 				{
-					const FString& PackageString = FPackageName::PackageFromPath(*Filename);
-					if (FCrc::StrCrc_DEPRECATED(*PackageString.ToUpper()) % SubsetMod != SubsetTarget)
+					const FString& SubPackageName = FPackageName::PackageFromPath(*Filename);
+					if (FCrc::StrCrc_DEPRECATED(*SubPackageName.ToUpper()) % SubsetMod != SubsetTarget)
 					{
 						continue;
 					}
@@ -226,20 +226,20 @@ int32 UDerivedDataCacheCommandlet::Main( const FString& Params )
 			GRedirectCollector.ResolveStringAssetReference();
 
 			// cache all the resources for this platform
-			/*for (TObjectIterator<UObject> It; It; ++It)
+			for (TObjectIterator<UObject> It; It; ++It)
 			{
 				if ((PackageFilter&NORMALIZE_ExcludeEnginePackages) == 0 || !It->GetOutermost()->GetName().StartsWith(TEXT("/Engine")))
 				{
 					if (!ProcessedPackages.Contains(It->GetOutermost()->GetFName()))
 					{
-						check( (It->GetOutermost()->GetPackageFlags() & PKG_ReloadingForCooker) == 0 );
+						check((It->GetOutermost()->GetPackageFlags() & PKG_ReloadingForCooker) == 0);
 						for (auto Platform : Platforms)
 						{
 							It->BeginCacheForCookedPlatformData(Platform);
 						}
 					}
 				}
-			}*/
+			}
 
 
 			// Keep track of which packages have already been processed along with the map.
@@ -255,20 +255,20 @@ int32 UDerivedDataCacheCommandlet::Main( const FString& Params )
 						continue;
 					}
 					if (!ProcessedPackages.Contains(Pkg->GetFName()))
-					{
-						ProcessedPackages.Add(Pkg->GetFName());
-						Pkg->SetPackageFlags(PKG_ReloadingForCooker);
 						{
-							TArray<UObject *> ObjectsInPackage;
-							GetObjectsWithOuter(Pkg, ObjectsInPackage, true);
-							for (int32 IndexPackage = 0; IndexPackage < ObjectsInPackage.Num(); IndexPackage++)
+						ProcessedPackages.Add(Pkg->GetFName());
+							Pkg->SetPackageFlags(PKG_ReloadingForCooker);
 							{
-								ObjectsInPackage[IndexPackage]->WillNeverCacheCookedPlatformDataAgain();
-								ObjectsInPackage[IndexPackage]->ClearAllCachedCookedPlatformData();
+								TArray<UObject *> ObjectsInPackage;
+								GetObjectsWithOuter(Pkg, ObjectsInPackage, true);
+								for (int32 IndexPackage = 0; IndexPackage < ObjectsInPackage.Num(); IndexPackage++)
+								{
+									ObjectsInPackage[IndexPackage]->WillNeverCacheCookedPlatformDataAgain();
+									ObjectsInPackage[IndexPackage]->ClearAllCachedCookedPlatformData();
+								}
 							}
 						}
 					}
-				}
 				FindProcessedPackagesTime += FPlatformTime::Seconds() - FindProcessedPackagesStartTime;
 			}
 
