HTNPlanner: Disable cplusplus.PlacementNew clang checker

#rnx

diff --git a/Engine/Plugins/AI/HTNPlanner/Source/HTNPlanner/HTNPlanner.Build.cs b/Engine/Plugins/AI/HTNPlanner/Source/HTNPlanner/HTNPlanner.Build.cs
--- a/Engine/Plugins/AI/HTNPlanner/Source/HTNPlanner/HTNPlanner.Build.cs
+++ b/Engine/Plugins/AI/HTNPlanner/Source/HTNPlanner/HTNPlanner.Build.cs
@@ -6,7 +6,9 @@ namespace UnrealBuildTool.Rules
     {
         public HTNPlanner(ReadOnlyTargetRules Target) : base(Target)
         {
-            PublicDependencyModuleNames.AddRange(
+			StaticAnalyzerDisabledCheckers.Add("cplusplus.PlacementNew");
+
+			PublicDependencyModuleNames.AddRange(
                 new string[] {
                         "Core",
                         "CoreUObject",
