[HTN] Static Analysis fixes, making sure the HTN domain allocation is aligned properly

#rnx

diff --git a/Engine/Plugins/AI/HTNPlanner/Source/HTNPlanner/Private/HTNBuilder.cpp b/Engine/Plugins/AI/HTNPlanner/Source/HTNPlanner/Private/HTNBuilder.cpp
--- a/Engine/Plugins/AI/HTNPlanner/Source/HTNPlanner/Private/HTNBuilder.cpp
+++ b/Engine/Plugins/AI/HTNPlanner/Source/HTNPlanner/Private/HTNBuilder.cpp
@@ -66,7 +66,10 @@ bool FHTNBuilder_Domain::Compile()
 		}
 	}
 
-	uint8* CompiledData = new uint8[MemoryRequired];
+	constexpr std::size_t MemoryAlignment = FMath::Max(FMath::Max3(alignof(FHTNPrimitiveTask), alignof(FHTNEffect), alignof(FHTNMethod)),
+		FMath::Max(alignof(FHTNPolicy::FTaskID), alignof(FHTNCondition)));
+
+	uint8* CompiledData = static_cast<uint8*>(FMemory::Malloc(MemoryRequired, MemoryAlignment));
 	uint8* CurrentAddress = CompiledData;
 	
 	for (auto It = PrimitiveTasks.CreateConstIterator(); It; ++It)
@@ -179,8 +182,8 @@ bool FHTNBuilder_Domain::Compile()
 
 	if (bIssuesFound)
 	{
-		delete[] CompiledData;
-		Domain.RawData = nullptr;
+		Domain.ReleaseRawData();
+		CompiledData = nullptr;
 	}
 	else if (MemoryUsed > 0)	// meaning we have anything scored
 	{
diff --git a/Engine/Plugins/AI/HTNPlanner/Source/HTNPlanner/Private/HTNDomain.cpp b/Engine/Plugins/AI/HTNPlanner/Source/HTNPlanner/Private/HTNDomain.cpp
--- a/Engine/Plugins/AI/HTNPlanner/Source/HTNPlanner/Private/HTNDomain.cpp
+++ b/Engine/Plugins/AI/HTNPlanner/Source/HTNPlanner/Private/HTNDomain.cpp
@@ -261,14 +261,22 @@ FHTNDomain::FHTNDomain()
 
 FHTNDomain::~FHTNDomain()
 {
-	delete[] RawData;
+	ReleaseRawData();
+}
+
+void FHTNDomain::ReleaseRawData()
+{
+	if (RawData)
+	{
+		FMemory::Free(RawData);
+		RawData = nullptr;
+	}
 }
 
 void FHTNDomain::Reset()
 {
 	//Tasks.Reset();
-	delete[] RawData;
-	RawData = nullptr;
+	ReleaseRawData();
 	TaskNameMap.Reset();
 	FirstCompositeTaskID = FHTNPolicy::InvalidTaskID;
 	RootTaskID = FHTNPolicy::InvalidTaskID;
@@ -283,7 +291,6 @@ FHTNPolicy::FTaskID FHTNDomain::FindTaskID(const FName& TaskName) const
 	return TaskID ? *TaskID : FHTNPolicy::InvalidTaskID;
 }
 
-
 FName FHTNDomain::GetTaskName(const FHTNPolicy::FTaskID TaskID) const
 {
 #if WITH_HTN_DEBUG
diff --git a/Engine/Plugins/AI/HTNPlanner/Source/HTNPlanner/Public/HTNDomain.h b/Engine/Plugins/AI/HTNPlanner/Source/HTNPlanner/Public/HTNDomain.h
--- a/Engine/Plugins/AI/HTNPlanner/Source/HTNPlanner/Public/HTNDomain.h
+++ b/Engine/Plugins/AI/HTNPlanner/Source/HTNPlanner/Public/HTNDomain.h
@@ -181,6 +181,8 @@ struct HTNPLANNER_API FHTNDomain : public TSharedFromThis<FHTNDomain>
 	FHTNDomain(); 
 	~FHTNDomain();
 
+	UE_NONCOPYABLE(FHTNDomain);
+
 	FHTNPolicy::FTaskID FindTaskID(const FName& TaskName) const;
 	FORCEINLINE bool IsPrimitiveTask(const FHTNPolicy::FTaskID TaskID) const
 	{
@@ -209,13 +211,14 @@ struct HTNPLANNER_API FHTNDomain : public TSharedFromThis<FHTNDomain>
 
 protected:
 	friend struct FHTNBuilder_Domain;
+	void ReleaseRawData();
 
-	//TArray<FHTNTaskWrapper> Tasks;
-	uint8* RawData;
+	uint8* RawData = nullptr;
 	TMap<FName, FHTNPolicy::FTaskID> TaskNameMap;
 	FHTNPolicy::FTaskID FirstCompositeTaskID;
-	/** Root task is the default task we start planning from
-	 *	*/
+	/**
+	 * Root task is the default task we start planning from
+	 */
 	FHTNPolicy::FTaskID RootTaskID;
 #if WITH_HTN_DEBUG
 	TMap<FHTNPolicy::FTaskID, FName> TaskIDToName;
