Fixed HTNPlanner's Build.cs

#rnx
#Synced-CL 49265344

diff --git a/Engine/Plugins/AI/HTNPlanner/Source/HTNPlanner/HTNPlanner.Build.cs b/Engine/Plugins/AI/HTNPlanner/Source/HTNPlanner/HTNPlanner.Build.cs
--- a/Engine/Plugins/AI/HTNPlanner/Source/HTNPlanner/HTNPlanner.Build.cs
+++ b/Engine/Plugins/AI/HTNPlanner/Source/HTNPlanner/HTNPlanner.Build.cs
@@ -2,30 +2,28 @@
 
 namespace UnrealBuildTool.Rules
 {
-    public class HTNPlanner : ModuleRules
-    {
-        public HTNPlanner(ReadOnlyTargetRules Target) : base(Target)
-        {
-			StaticAnalyzerDisabledCheckers.Add("cplusplus.PlacementNew");
-
+	public class HTNPlanner : ModuleRules
+	{
+		public HTNPlanner(ReadOnlyTargetRules Target) : base(Target)
+		{
 			PublicDependencyModuleNames.AddRange(
-                new string[] {
-                        "Core",
-                        "CoreUObject",
-                        "Engine",
-                        "GameplayTags",
-                        "GameplayTasks",
-                        "AIModule"
-                }
-                );
+				new string[] {
+					"Core",
+					"CoreUObject",
+					"Engine",
+					"GameplayTags",
+					"GameplayTasks",
+					"AIModule"
+				}
+			);
 
-            if (Target.bBuildEditor == true)
-            {
-                PrivateDependencyModuleNames.Add("EditorFramework");
-                PrivateDependencyModuleNames.Add("UnrealEd");
-            }
+			if (Target.bBuildEditor == true)
+			{
+				PrivateDependencyModuleNames.Add("EditorFramework");
+				PrivateDependencyModuleNames.Add("UnrealEd");
+			}
 
-            SetupGameplayDebuggerSupport(Target);
-        }
-    }
+			SetupGameplayDebuggerSupport(Target);
+		}
+	}
 }
